Fixed crash on CastChecked in FMatinee::SetInterpEnd().

#ttp 333548 - MATINEE: CRASH: Moving the End of Timeline Red Marker will crash the Editor
#branch UE4
#proj Editor.Matinee

#fix Changed CastChecked to Cast in FMatinee::SetInterpEnd().

diff --git a/Engine/Source/Editor/Matinee/Private/MatineeTools.cpp b/Engine/Source/Editor/Matinee/Private/MatineeTools.cpp
--- a/Engine/Source/Editor/Matinee/Private/MatineeTools.cpp
+++ b/Engine/Source/Editor/Matinee/Private/MatineeTools.cpp
@@ -1787,7 +1787,7 @@ void FMatinee::SetInterpEnd(float NewInterpLength)
 	CurveEd->SetRegionMarker(true, IData->EdSectionStart, IData->EdSectionEnd, RegionFillColor);
 
 	// Update the CameraAnim if necessary
-	AMatineeActorCameraAnim* const CamAnimMatineeActor = CastChecked<AMatineeActorCameraAnim>(MatineeActor);
+	AMatineeActorCameraAnim* const CamAnimMatineeActor = Cast<AMatineeActorCameraAnim>(MatineeActor);
 	if ( (CamAnimMatineeActor != nullptr) && (CamAnimMatineeActor->CameraAnim != nullptr) )
 	{
 		CamAnimMatineeActor->CameraAnim->AnimLength = IData->InterpLength;
