Allow director track in Matinee to cut to itself for smooth transition workaround

diff --git a/Engine/Source/Editor/Matinee/Private/MatineeTrackHelpers.cpp b/Engine/Source/Editor/Matinee/Private/MatineeTrackHelpers.cpp
--- a/Engine/Source/Editor/Matinee/Private/MatineeTrackHelpers.cpp
+++ b/Engine/Source/Editor/Matinee/Private/MatineeTrackHelpers.cpp
@@ -304,6 +304,7 @@ UMatineeTrackDirectorHelper::UMatineeTrackDirectorHelper(const class FPostConstr
 }
 
 bool UMatineeTrackDirectorHelper::PreCreateKeyframe( UInterpTrack *Track, float KeyTime ) const
+
 {
 	// If adding a cut, bring up combo to let user choose group to cut to.
 	KeyframeAddDataName = NAME_None;
@@ -318,8 +319,7 @@ bool UMatineeTrackDirectorHelper::PreCreateKeyframe( UInterpTrack *Track, float
 		TArray<FString> GroupNames;
 		for ( UInterpGroup* InterpGroup : Mode->InterpEd->GetInterpData()->InterpGroups )
 		{
-			// Skip folder groups and the director group (we can't cut to ourselves)
-			if ( !(InterpGroup->bIsFolder || InterpGroup->IsA(UInterpGroupDirector::StaticClass())) )
+			if ( !InterpGroup->bIsFolder )
 			{
 				GroupNames.Add(InterpGroup->GroupName.ToString());
 			}
