Remove Matinee related headers from Engine.h, and add only where needed

diff --git a/Engine/Source/Developer/AssetTools/Private/AssetTypeActions/AssetTypeActions_CameraAnim.cpp b/Engine/Source/Developer/AssetTools/Private/AssetTypeActions/AssetTypeActions_CameraAnim.cpp
--- a/Engine/Source/Developer/AssetTools/Private/AssetTypeActions/AssetTypeActions_CameraAnim.cpp
+++ b/Engine/Source/Developer/AssetTools/Private/AssetTypeActions/AssetTypeActions_CameraAnim.cpp
@@ -2,6 +2,12 @@
 
 #include "AssetToolsPrivatePCH.h"
 #include "Toolkits/IToolkitHost.h"
+#include "Matinee/InterpData.h"
+#include "Matinee/InterpGroupCamera.h"
+#include "Matinee/InterpGroupInst.h"
+#include "Matinee/InterpTrackInst.h"
+#include "Matinee/InterpTrackAnimControl.h"
+#include "Matinee/MatineeActorCameraAnim.h"
 
 DEFINE_LOG_CATEGORY_STATIC(LogCameraAnim, Log, All);
 
diff --git a/Engine/Source/Developer/AssetTools/Private/AssetTypeActions/AssetTypeActions_InterpData.h b/Engine/Source/Developer/AssetTools/Private/AssetTypeActions/AssetTypeActions_InterpData.h
--- a/Engine/Source/Developer/AssetTools/Private/AssetTypeActions/AssetTypeActions_InterpData.h
+++ b/Engine/Source/Developer/AssetTools/Private/AssetTypeActions/AssetTypeActions_InterpData.h
@@ -2,6 +2,8 @@
 
 #pragma once
 
+#include "Matinee/InterpData.h"
+
 class FAssetTypeActions_InterpData : public FAssetTypeActions_Base
 {
 public:
diff --git a/Engine/Source/Editor/BlueprintGraph/Private/K2ActionMenuBuilder.cpp b/Engine/Source/Editor/BlueprintGraph/Private/K2ActionMenuBuilder.cpp
--- a/Engine/Source/Editor/BlueprintGraph/Private/K2ActionMenuBuilder.cpp
+++ b/Engine/Source/Editor/BlueprintGraph/Private/K2ActionMenuBuilder.cpp
@@ -5,6 +5,7 @@
 #include "ComponentAssetBroker.h"
 #include "KismetEditorUtilities.h"
 #include "ObjectEditorUtils.h"
+#include "Matinee/MatineeActor.h"
 
 #define LOCTEXT_NAMESPACE "KismetSchema"
 
diff --git a/Engine/Source/Editor/BlueprintGraph/Private/K2Node_MatineeController.cpp b/Engine/Source/Editor/BlueprintGraph/Private/K2Node_MatineeController.cpp
--- a/Engine/Source/Editor/BlueprintGraph/Private/K2Node_MatineeController.cpp
+++ b/Engine/Source/Editor/BlueprintGraph/Private/K2Node_MatineeController.cpp
@@ -3,6 +3,8 @@
 #include "BlueprintGraphPrivatePCH.h"
 #include "KismetCompiler.h"
 #include "Runtime/Engine/Public/MatineeDelegates.h"
+#include "Matinee/MatineeActor.h"
+#include "Matinee/InterpData.h"
 
 UK2Node_MatineeController::UK2Node_MatineeController(const class FPostConstructInitializeProperties& PCIP)
 	: Super(PCIP)
diff --git a/Engine/Source/Editor/DetailCustomizations/Private/MatineeActorDetails.cpp b/Engine/Source/Editor/DetailCustomizations/Private/MatineeActorDetails.cpp
--- a/Engine/Source/Editor/DetailCustomizations/Private/MatineeActorDetails.cpp
+++ b/Engine/Source/Editor/DetailCustomizations/Private/MatineeActorDetails.cpp
@@ -1,6 +1,7 @@
 // Copyright 1998-2014 Epic Games, Inc. All Rights Reserved.
 
 #include "DetailCustomizationsPrivatePCH.h"
+#include "Matinee/MatineeActor.h"
 #include "MatineeActorDetails.h"
 
 #define LOCTEXT_NAMESPACE "MatineeActorDetails"
diff --git a/Engine/Source/Editor/LevelEditor/Private/LevelEditorActions.cpp b/Engine/Source/Editor/LevelEditor/Private/LevelEditorActions.cpp
--- a/Engine/Source/Editor/LevelEditor/Private/LevelEditorActions.cpp
+++ b/Engine/Source/Editor/LevelEditor/Private/LevelEditorActions.cpp
@@ -2,6 +2,7 @@
 
 
 #include "LevelEditor.h"
+#include "Matinee/MatineeActor.h"
 #include "EditorSupportDelegates.h"
 #include "SLevelEditor.h"
 #include "LevelEditorActions.h"
diff --git a/Engine/Source/Editor/LevelEditor/Private/LevelEditorContextMenu.cpp b/Engine/Source/Editor/LevelEditor/Private/LevelEditorContextMenu.cpp
--- a/Engine/Source/Editor/LevelEditor/Private/LevelEditorContextMenu.cpp
+++ b/Engine/Source/Editor/LevelEditor/Private/LevelEditorContextMenu.cpp
@@ -2,6 +2,7 @@
 
 
 #include "LevelEditor.h"
+#include "Matinee/MatineeActor.h"
 #include "BlueprintUtilities.h"
 #include "Editor/Kismet/Public/BlueprintEditorModule.h"
 #include "Editor/UnrealEd/Public/Kismet2/KismetEditorUtilities.h"
diff --git a/Engine/Source/Editor/LevelEditor/Private/LevelEditorToolBar.cpp b/Engine/Source/Editor/LevelEditor/Private/LevelEditorToolBar.cpp
--- a/Engine/Source/Editor/LevelEditor/Private/LevelEditorToolBar.cpp
+++ b/Engine/Source/Editor/LevelEditor/Private/LevelEditorToolBar.cpp
@@ -14,6 +14,7 @@
 #include "AssetData.h"
 #include "Editor/ContentBrowser/Public/ContentBrowserModule.h"
 #include "EngineBuildSettings.h"
+#include "Matinee/MatineeActor.h"
 
 static TAutoConsoleVariable<int32> CVarFeatureLevelPreview(
 	TEXT("r.FeatureLevelPreview"),
diff --git a/Engine/Source/Editor/Levels/Private/LevelCollectionViewModel.cpp b/Engine/Source/Editor/Levels/Private/LevelCollectionViewModel.cpp
--- a/Engine/Source/Editor/Levels/Private/LevelCollectionViewModel.cpp
+++ b/Engine/Source/Editor/Levels/Private/LevelCollectionViewModel.cpp
@@ -1,6 +1,8 @@
 // Copyright 1998-2014 Epic Games, Inc. All Rights Reserved.
 #include "LevelsPrivatePCH.h"
 
+#include "Matinee/MatineeActor.h"
+
 #include "EditorLevelUtils.h"
 #include "LevelUtils.h"
 #include "Editor/PropertyEditor/Public/PropertyEditorModule.h"
diff --git a/Engine/Source/Editor/Matinee/Private/Matinee.cpp b/Engine/Source/Editor/Matinee/Private/Matinee.cpp
--- a/Engine/Source/Editor/Matinee/Private/Matinee.cpp
+++ b/Engine/Source/Editor/Matinee/Private/Matinee.cpp
@@ -5,6 +5,18 @@
 =============================================================================*/
 
 #include "MatineeModule.h"
+
+#include "Matinee/InterpFilter.h"
+#include "Matinee/MatineeActor.h"
+#include "Matinee/MatineeActorCameraAnim.h"
+#include "Matinee/InterpGroupInst.h"
+#include "Matinee/InterpTrackToggle.h"
+#include "Matinee/InterpTrackSound.h"
+#include "Matinee/InterpTrackDirector.h"
+#include "Matinee/InterpTrackVisibility.h"
+#include "Matinee/InterpTrackEvent.h"
+#include "Matinee/InterpGroupDirector.h"
+
 #include "Matinee.h"
 #include "MatineeActions.h"
 #include "MatineeFilterButton.h"
@@ -14,7 +26,6 @@
 
 #include "MatineeClasses.h"
 
-
 #include "CameraController.h"
 #include "MatineeConstants.h"
 
diff --git a/Engine/Source/Editor/Matinee/Private/Matinee.h b/Engine/Source/Editor/Matinee/Private/Matinee.h
--- a/Engine/Source/Editor/Matinee/Private/Matinee.h
+++ b/Engine/Source/Editor/Matinee/Private/Matinee.h
@@ -2,6 +2,8 @@
 
 #pragma once
 
+#include "Matinee/InterpData.h"
+
 #include "IMatinee.h"
 #include "IDistCurveEditor.h"
 #include "MatineeActions.h"
@@ -276,7 +278,7 @@ public:
 	void OnContextSetKeyTime();
 	void OnContextSetKeyTimeTextCommitted(const FText& InText, ETextCommit::Type CommitInfo, FInterpEdSelKey* SelKey, UInterpTrack* Track);
 	void OnContextSetValue();
-	void OnContextSetValueTextCommitted(const FText& InText, ETextCommit::Type CommitInfo, FInterpEdSelKey* SelKey, UInterpTrackFloatBase* FloatTrack);
+	void OnContextSetValueTextCommitted(const FText& InText, ETextCommit::Type CommitInfo, FInterpEdSelKey* SelKey, class UInterpTrackFloatBase* FloatTrack);
 
 	/** Pops up a menu and lets you set the color for the selected key. Not all track types are supported. */
 	void OnContextSetColor();
@@ -291,16 +293,16 @@ public:
 
 	/** Pops up menu and lets the user set a group to use to lookup transform info for a movement keyframe. */
 	void OnSetMoveKeyLookupGroup();
-	void OnSetMoveKeyLookupGroupTextChosen(const FString& ChosenText, int32 KeyIndex, UInterpTrackMove* MoveTrack, UInterpTrackMoveAxis* MoveTrackAxis);
+	void OnSetMoveKeyLookupGroupTextChosen(const FString& ChosenText, int32 KeyIndex, UInterpTrackMove* MoveTrack, class UInterpTrackMoveAxis* MoveTrackAxis);
 
 	/** Clears the lookup group for a currently selected movement key. */
 	void OnClearMoveKeyLookupGroup();
 
 	void OnSetAnimKeyLooping( bool bInLooping );
 	void OnSetAnimOffset( bool bInEndOffset );
-	void OnSetAnimOffsetTextCommitted(const FText& InText, ETextCommit::Type CommitInfo, FInterpEdSelKey* SelKey, UInterpTrackAnimControl* AnimTrack, bool bEndOffset);
+	void OnSetAnimOffsetTextCommitted(const FText& InText, ETextCommit::Type CommitInfo, FInterpEdSelKey* SelKey, class UInterpTrackAnimControl* AnimTrack, bool bEndOffset);
 	void OnSetAnimPlayRate();
-	void OnSetAnimPlayRateTextCommitted(const FText& InText, ETextCommit::Type CommitInfo, FInterpEdSelKey* SelKey, UInterpTrackAnimControl* AnimTrack);
+	void OnSetAnimPlayRateTextCommitted(const FText& InText, ETextCommit::Type CommitInfo, FInterpEdSelKey* SelKey, class UInterpTrackAnimControl* AnimTrack);
 
 	/** Handler for the toggle animation reverse menu item. */
 	void OnToggleReverseAnim();
@@ -341,9 +343,9 @@ public:
 	void OnSetSoundPitch();
 	void OnSetSoundPitchTextEntered(const FText& InText, ETextCommit::Type CommitInfo, TArray<int32> SoundTrackKeyIndices);
 	void OnContextDirKeyTransitionTime();
-	void OnContextDirKeyTransitionTimeTextCommitted(const FText& InText, ETextCommit::Type CommitInfo, FInterpEdSelKey* SelKey, UInterpTrackDirector* DirTrack);
+	void OnContextDirKeyTransitionTimeTextCommitted(const FText& InText, ETextCommit::Type CommitInfo, FInterpEdSelKey* SelKey, class UInterpTrackDirector* DirTrack);
 	void OnContextDirKeyRenameCameraShot();
-	void OnContextDirKeyRenameCameraShotTextCommitted(const FText& InText, ETextCommit::Type CommitInfo, FInterpEdSelKey* SelKey, UInterpTrackDirector* DirTrack);
+	void OnContextDirKeyRenameCameraShotTextCommitted(const FText& InText, ETextCommit::Type CommitInfo, FInterpEdSelKey* SelKey, class UInterpTrackDirector* DirTrack);
 	void OnFlipToggleKey();
 
 	/** Called when a new key condition is selected in a track keyframe context menu */
@@ -363,11 +365,11 @@ public:
 
 	/** Called when the user wants to set the clip ID number for Particle Replay track keys */
 	void OnParticleReplayKeyContext_SetClipIDNumber();
-	void OnParticleReplayKeyContext_SetClipIDNumberTextCommitted(const FText& InText, ETextCommit::Type CommitInfo, FParticleReplayTrackKey* ParticleReplayKey);
+	void OnParticleReplayKeyContext_SetClipIDNumberTextCommitted(const FText& InText, ETextCommit::Type CommitInfo, struct FParticleReplayTrackKey* ParticleReplayKey);
 
 	/** Called when the user wants to set the duration of Particle Replay track keys */
 	void OnParticleReplayKeyContext_SetDuration();
-	void OnParticleReplayKeyContext_SetDurationTextCommitted(const FText& InText, ETextCommit::Type CommitInfo, FParticleReplayTrackKey* ParticleReplayKey);
+	void OnParticleReplayKeyContext_SetDurationTextCommitted(const FText& InText, ETextCommit::Type CommitInfo, struct FParticleReplayTrackKey* ParticleReplayKey);
 
 	/** Called to delete the currently selected keys */
 	void OnDeleteSelectedKeys();
diff --git a/Engine/Source/Editor/Matinee/Private/MatineeCurveEd.cpp b/Engine/Source/Editor/Matinee/Private/MatineeCurveEd.cpp
--- a/Engine/Source/Editor/Matinee/Private/MatineeCurveEd.cpp
+++ b/Engine/Source/Editor/Matinee/Private/MatineeCurveEd.cpp
@@ -7,6 +7,12 @@
 
 
 #include "MatineeModule.h"
+
+#include "Matinee/MatineeActor.h"
+#include "Matinee/InterpGroupInst.h"
+#include "Matinee/InterpTrackMoveAxis.h"
+#include "Matinee/InterpTrackColorProp.h"
+
 #include "Matinee.h"
 
 /** Ensure the curve editor is synchronised with the track editor. */
diff --git a/Engine/Source/Editor/Matinee/Private/MatineeDraw.cpp b/Engine/Source/Editor/Matinee/Private/MatineeDraw.cpp
--- a/Engine/Source/Editor/Matinee/Private/MatineeDraw.cpp
+++ b/Engine/Source/Editor/Matinee/Private/MatineeDraw.cpp
@@ -11,6 +11,14 @@
 #include "Runtime/Engine/Public/Interpolation.h"
 #include "Runtime/Engine/Public/InterpolationHitProxy.h"
 #include "Runtime/Engine/Public/Slate/SceneViewport.h"
+#include "Matinee/MatineeActor.h"
+#include "Matinee/InterpGroupDirector.h"
+#include "Matinee/InterpGroupInst.h"
+#include "Matinee/InterpTrackFloatBase.h"
+#include "Matinee/InterpTrackVectorBase.h"
+#include "Matinee/InterpTrackLinearColorBase.h"
+#include "Matinee/InterpTrackEvent.h"
+#include "Matinee/InterpTrackMove.h"
 
 static const int32 GroupHeadHeight = 24;
 static const int32 TrackHeight = 24;
diff --git a/Engine/Source/Editor/Matinee/Private/MatineeKeyReduction.cpp b/Engine/Source/Editor/Matinee/Private/MatineeKeyReduction.cpp
--- a/Engine/Source/Editor/Matinee/Private/MatineeKeyReduction.cpp
+++ b/Engine/Source/Editor/Matinee/Private/MatineeKeyReduction.cpp
@@ -5,6 +5,8 @@
 =============================================================================*/
 
 #include "MatineeModule.h"
+#include "Matinee/MatineeActor.h"
+#include "Matinee/InterpGroupInst.h"
 #include "Matinee.h"
 #include "Slate.h"
 #include "SNumericEntryBox.h"
diff --git a/Engine/Source/Editor/Matinee/Private/MatineeMenus.cpp b/Engine/Source/Editor/Matinee/Private/MatineeMenus.cpp
--- a/Engine/Source/Editor/Matinee/Private/MatineeMenus.cpp
+++ b/Engine/Source/Editor/Matinee/Private/MatineeMenus.cpp
@@ -5,6 +5,30 @@
 =============================================================================*/
 
 #include "MatineeModule.h"
+
+#include "Matinee/InterpTrackInst.h"
+#include "Matinee/InterpTrackMove.h"
+#include "Matinee/InterpTrackMoveAxis.h"
+#include "Matinee/InterpTrackInstMove.h"
+#include "Matinee/InterpTrackEvent.h"
+#include "Matinee/InterpTrackToggle.h"
+#include "Matinee/InterpTrackSound.h"
+#include "Matinee/InterpTrackVisibility.h"
+#include "Matinee/InterpTrackBoolProp.h"
+#include "Matinee/InterpTrackFloatProp.h"
+#include "Matinee/InterpTrackColorProp.h"
+#include "Matinee/InterpTrackVectorProp.h"
+#include "Matinee/InterpTrackLinearColorProp.h"
+#include "Matinee/InterpTrackAnimControl.h"
+#include "Matinee/InterpTrackParticleReplay.h"
+#include "Matinee/InterpTrackVectorMaterialParam.h"
+#include "Matinee/InterpTrackDirector.h"
+#include "Matinee/InterpTrackInstDirector.h"
+#include "Matinee/InterpTrackAudioMaster.h"
+#include "Matinee/InterpGroupDirector.h"
+#include "Matinee/InterpGroupInstDirector.h"
+#include "Matinee/InterpFilter_Custom.h"
+
 #include "Matinee.h"
 #include "MatineeDelegates.h"
 
diff --git a/Engine/Source/Editor/Matinee/Private/MatineeModule.cpp b/Engine/Source/Editor/Matinee/Private/MatineeModule.cpp
--- a/Engine/Source/Editor/Matinee/Private/MatineeModule.cpp
+++ b/Engine/Source/Editor/Matinee/Private/MatineeModule.cpp
@@ -3,6 +3,7 @@
 #include "MatineeModule.h"
 #include "ModuleManager.h"
 #include "Matinee.h"
+#include "Matinee/MatineeActor.h"
 
 const FName MatineeAppIdentifier = FName(TEXT("MatineeApp"));
 
diff --git a/Engine/Source/Editor/Matinee/Private/MatineeRecordMovie.cpp b/Engine/Source/Editor/Matinee/Private/MatineeRecordMovie.cpp
--- a/Engine/Source/Editor/Matinee/Private/MatineeRecordMovie.cpp
+++ b/Engine/Source/Editor/Matinee/Private/MatineeRecordMovie.cpp
@@ -5,6 +5,7 @@
 =============================================================================*/
 
 #include "MatineeModule.h"
+#include "Matinee/MatineeActor.h"
 #include "Matinee.h"
 
 #define LOCTEXT_NAMESPACE "MatineeMovieCapture"
diff --git a/Engine/Source/Editor/Matinee/Private/MatineeTools.cpp b/Engine/Source/Editor/Matinee/Private/MatineeTools.cpp
--- a/Engine/Source/Editor/Matinee/Private/MatineeTools.cpp
+++ b/Engine/Source/Editor/Matinee/Private/MatineeTools.cpp
@@ -4,10 +4,34 @@
 	MatineeTools.cpp: Interpolation editing support tools
 =============================================================================*/
 
+
 #include "MatineeModule.h"
 #include "Matinee.h"
 #include "MatineeClasses.h"
 
+#include "Matinee/MatineeActor.h"
+#include "Matinee/MatineeActorCameraAnim.h"
+#include "Matinee/InterpTrackInst.h"
+#include "Matinee/InterpTrackMove.h"
+#include "Matinee/InterpTrackMoveAxis.h"
+#include "Matinee/InterpTrackInstMove.h"
+#include "Matinee/InterpTrackEvent.h"
+#include "Matinee/InterpTrackFloatProp.h"
+#include "Matinee/InterpTrackVectorBase.h"
+#include "Matinee/InterpTrackLinearColorBase.h"
+#include "Matinee/InterpTrackSound.h"
+#include "Matinee/InterpTrackDirector.h"
+#include "Matinee/InterpTrackFade.h"
+#include "Matinee/InterpTrackSlomo.h"
+#include "Matinee/InterpTrackColorScale.h"
+#include "Matinee/InterpTrackInstDirector.h"
+#include "Matinee/InterpTrackAnimControl.h"
+#include "Matinee/InterpTrackParticleReplay.h"
+#include "Matinee/InterpGroupInst.h"
+#include "Matinee/InterpGroupDirector.h"
+#include "Matinee/InterpGroupInstDirector.h"
+#include "Matinee/InterpFilter.h"
+
 #include "MatineeConstants.h"
 #include "MatineeViewSaveData.h"
 #include "MatineeDelegates.h"
diff --git a/Engine/Source/Editor/Matinee/Private/MatineeTrackHelpers.cpp b/Engine/Source/Editor/Matinee/Private/MatineeTrackHelpers.cpp
--- a/Engine/Source/Editor/Matinee/Private/MatineeTrackHelpers.cpp
+++ b/Engine/Source/Editor/Matinee/Private/MatineeTrackHelpers.cpp
@@ -8,6 +8,19 @@
 #include "SoundDefinitions.h"
 #include "MatineeUtils.h"
 #include "ContentBrowserModule.h"
+#include "Matinee/MatineeActor.h"
+#include "Matinee/InterpTrackAnimControl.h"
+#include "Matinee/InterpTrackDirector.h"
+#include "Matinee/InterpTrackSound.h"
+#include "Matinee/InterpTrackEvent.h"
+#include "Matinee/InterpTrackToggle.h"
+#include "Matinee/InterpTrackVisibility.h"
+#include "Matinee/InterpTrackBoolProp.h"
+#include "Matinee/InterpTrackColorProp.h"
+#include "Matinee/InterpTrackFloatProp.h"
+#include "Matinee/InterpTrackVectorProp.h"
+#include "Matinee/InterpTrackLinearColorProp.h"
+#include "Matinee/InterpGroupInst.h"
 
 #define LOCTEXT_NAMESPACE "MatineeTrackHelpers"
 
diff --git a/Engine/Source/Editor/Matinee/Private/MatineeTransaction.cpp b/Engine/Source/Editor/Matinee/Private/MatineeTransaction.cpp
--- a/Engine/Source/Editor/Matinee/Private/MatineeTransaction.cpp
+++ b/Engine/Source/Editor/Matinee/Private/MatineeTransaction.cpp
@@ -5,6 +5,12 @@
 #include "MatineeTransaction.h"
 #include "Editor/BlueprintGraph/Classes/K2Node.h"
 #include "Editor/BlueprintGraph/Classes/K2Node_MatineeController.h"
+#include "Matinee/MatineeActor.h"
+#include "Matinee/InterpData.h"
+#include "Matinee/InterpGroup.h"
+#include "Matinee/InterpGroupInst.h"
+#include "Matinee/InterpTrack.h"
+#include "Matinee/InterpTrackInst.h"
 
 /*-----------------------------------------------------------------------------
 	FMatineeTransaction
diff --git a/Engine/Source/Editor/Matinee/Private/MatineeViewportClient.cpp b/Engine/Source/Editor/Matinee/Private/MatineeViewportClient.cpp
--- a/Engine/Source/Editor/Matinee/Private/MatineeViewportClient.cpp
+++ b/Engine/Source/Editor/Matinee/Private/MatineeViewportClient.cpp
@@ -5,6 +5,9 @@
 
 #include "Runtime/Engine/Public/InterpolationHitProxy.h"
 #include "Runtime/Engine/Public/Slate/SceneViewport.h"
+#include "Matinee/MatineeACtor.h"
+#include "Matinee/InterpTrackMove.h"
+#include "Matinee/InterpTrackEvent.h"
 
 /*-----------------------------------------------------------------------------
  FMatineeViewportClient
diff --git a/Engine/Source/Editor/Matinee/Public/IMatinee.h b/Engine/Source/Editor/Matinee/Public/IMatinee.h
--- a/Engine/Source/Editor/Matinee/Public/IMatinee.h
+++ b/Engine/Source/Editor/Matinee/Public/IMatinee.h
@@ -2,6 +2,11 @@
 
 #pragma once
 
+class AMatineeActor;
+class UInterpData;
+class UInterpGroup;
+class UInterpTrack;
+
 #include "MatineeGroupData.h"
 #include "MatineeTrackData.h"
 
diff --git a/Engine/Source/Editor/Matinee/Public/MatineeGroupData.h b/Engine/Source/Editor/Matinee/Public/MatineeGroupData.h
--- a/Engine/Source/Editor/Matinee/Public/MatineeGroupData.h
+++ b/Engine/Source/Editor/Matinee/Public/MatineeGroupData.h
@@ -3,6 +3,8 @@
 
 #pragma once
 
+#include "Matinee/InterpGroup.h"
+
 /*-----------------------------------------------------------------------------
 	FInterpGroup
 -----------------------------------------------------------------------------*/
diff --git a/Engine/Source/Editor/Matinee/Public/MatineeTrackData.h b/Engine/Source/Editor/Matinee/Public/MatineeTrackData.h
--- a/Engine/Source/Editor/Matinee/Public/MatineeTrackData.h
+++ b/Engine/Source/Editor/Matinee/Public/MatineeTrackData.h
@@ -2,6 +2,8 @@
 
 #pragma once
 
+#include "Matinee/InterpTrack.h"
+
 /*-----------------------------------------------------------------------------
 	Interp Track Filters
 -----------------------------------------------------------------------------*/
diff --git a/Engine/Source/Editor/UnrealEd/Private/AssetSelection.cpp b/Engine/Source/Editor/UnrealEd/Private/AssetSelection.cpp
--- a/Engine/Source/Editor/UnrealEd/Private/AssetSelection.cpp
+++ b/Engine/Source/Editor/UnrealEd/Private/AssetSelection.cpp
@@ -1,6 +1,7 @@
 // Copyright 1998-2014 Epic Games, Inc. All Rights Reserved.
 
 #include "UnrealEd.h"
+#include "Matinee/MatineeActor.h"
 #include "AssetSelection.h"
 #include "ScopedTransaction.h"
 
diff --git a/Engine/Source/Editor/UnrealEd/Private/Editor.cpp b/Engine/Source/Editor/UnrealEd/Private/Editor.cpp
--- a/Engine/Source/Editor/UnrealEd/Private/Editor.cpp
+++ b/Engine/Source/Editor/UnrealEd/Private/Editor.cpp
@@ -2,6 +2,7 @@
 
 
 #include "UnrealEd.h"
+#include "Matinee/MatineeActor.h"
 #include "EditorSupportDelegates.h"
 #include "Factories.h"
 #include "BSPOps.h"
diff --git a/Engine/Source/Editor/UnrealEd/Private/EditorServer.cpp b/Engine/Source/Editor/UnrealEd/Private/EditorServer.cpp
--- a/Engine/Source/Editor/UnrealEd/Private/EditorServer.cpp
+++ b/Engine/Source/Editor/UnrealEd/Private/EditorServer.cpp
@@ -21,6 +21,9 @@
 #include "../Private/GeomFitUtils.h"
 #include "Editor/GeometryMode/Public/GeometryEdMode.h"
 
+#include "Matinee/InterpTrackAnimControl.h"
+#include "Matinee/InterpData.h"
+
 #include "Editor/UnrealEd/Public/Kismet2/KismetEditorUtilities.h"
 #include "Editor/PropertyEditor/Public/PropertyEditorModule.h"
 #include "Editor/PropertyEditor/Public/IPropertyTable.h"
diff --git a/Engine/Source/Editor/UnrealEd/Private/EditorViewportClient.cpp b/Engine/Source/Editor/UnrealEd/Private/EditorViewportClient.cpp
--- a/Engine/Source/Editor/UnrealEd/Private/EditorViewportClient.cpp
+++ b/Engine/Source/Editor/UnrealEd/Private/EditorViewportClient.cpp
@@ -13,6 +13,8 @@
 #include "EngineAnalytics.h"
 #include "AnalyticsEventAttribute.h"
 #include "IAnalyticsProvider.h"
+#include "Matinee/MatineeActor.h"
+
 
 #define LOCTEXT_NAMESPACE "EditorViewportClient"
 
diff --git a/Engine/Source/Editor/UnrealEd/Private/Factories/ActorFactory.cpp b/Engine/Source/Editor/UnrealEd/Private/Factories/ActorFactory.cpp
--- a/Engine/Source/Editor/UnrealEd/Private/Factories/ActorFactory.cpp
+++ b/Engine/Source/Editor/UnrealEd/Private/Factories/ActorFactory.cpp
@@ -13,6 +13,8 @@ ActorFactory.cpp:
 #include "AssetData.h"
 #include "ScopedTransaction.h"
 #include "BSPOps.h"
+#include "Matinee/MatineeActor.h"
+#include "Matinee/InterpData.h"
 
 #include "AssetRegistryModule.h"
 
diff --git a/Engine/Source/Editor/UnrealEd/Private/Factories/EditorFactories.cpp b/Engine/Source/Editor/UnrealEd/Private/Factories/EditorFactories.cpp
--- a/Engine/Source/Editor/UnrealEd/Private/Factories/EditorFactories.cpp
+++ b/Engine/Source/Editor/UnrealEd/Private/Factories/EditorFactories.cpp
@@ -5,6 +5,8 @@
 =============================================================================*/
 
 #include "UnrealEd.h"
+#include "Matinee/InterpData.h"
+#include "Matinee/InterpGroupCamera.h"
 #include "Factories.h"
 #include "SoundDefinitions.h"
 #include "BlueprintUtilities.h"
diff --git a/Engine/Source/Editor/UnrealEd/Private/Fbx/FbxAnimationExport.cpp b/Engine/Source/Editor/UnrealEd/Private/Fbx/FbxAnimationExport.cpp
--- a/Engine/Source/Editor/UnrealEd/Private/Fbx/FbxAnimationExport.cpp
+++ b/Engine/Source/Editor/UnrealEd/Private/Fbx/FbxAnimationExport.cpp
@@ -33,6 +33,9 @@
 =============================================================================*/
 
 #include "UnrealEd.h"
+#include "Matinee/InterpData.h"
+#include "Matinee/InterpTrackAnimControl.h"
+
 
 DEFINE_LOG_CATEGORY_STATIC(LogFbxAnimationExport, Log, All);
 
diff --git a/Engine/Source/Editor/UnrealEd/Private/Fbx/FbxMainExport.cpp b/Engine/Source/Editor/UnrealEd/Private/Fbx/FbxMainExport.cpp
--- a/Engine/Source/Editor/UnrealEd/Private/Fbx/FbxMainExport.cpp
+++ b/Engine/Source/Editor/UnrealEd/Private/Fbx/FbxMainExport.cpp
@@ -33,6 +33,12 @@
 =============================================================================*/
 
 #include "UnrealEd.h"
+#include "Matinee/InterpData.h"
+#include "Matinee/InterpTrackMove.h"
+#include "Matinee/InterpTrackMoveAxis.h"
+#include "Matinee/InterpTrackFloatProp.h"
+#include "Matinee/InterpTrackInstFloatProp.h"
+#include "Matinee/InterpTrackInstMove.h"
 #include "LandscapeDataAccess.h"
 
 #include "FbxExporter.h"
diff --git a/Engine/Source/Editor/UnrealEd/Private/Fbx/FbxMatineeImport.cpp b/Engine/Source/Editor/UnrealEd/Private/Fbx/FbxMatineeImport.cpp
--- a/Engine/Source/Editor/UnrealEd/Private/Fbx/FbxMatineeImport.cpp
+++ b/Engine/Source/Editor/UnrealEd/Private/Fbx/FbxMatineeImport.cpp
@@ -29,6 +29,15 @@
 */
 
 #include "UnrealEd.h"
+#include "Matinee/MatineeActor.h"
+#include "Matinee/InterpGroup.h"
+#include "Matinee/InterpGroupInst.h"
+#include "Matinee/InterpData.h"
+#include "Matinee/InterpTrackMove.h"
+#include "Matinee/InterpTrackMoveAxis.h"
+#include "Matinee/InterpTrackInstMove.h"
+#include "Matinee/InterpTrackFloatProp.h"
+#include "Matinee/InterpTrackInstFloatProp.h"
 
 #include "FbxImporter.h"
 
diff --git a/Engine/Source/Editor/UnrealEd/Private/FbxImporter.h b/Engine/Source/Editor/UnrealEd/Private/FbxImporter.h
--- a/Engine/Source/Editor/UnrealEd/Private/FbxImporter.h
+++ b/Engine/Source/Editor/UnrealEd/Private/FbxImporter.h
@@ -914,7 +914,7 @@ protected:
 	/**
 	 * Imports an FBX transform curve into a movement subtrack
 	 */
-	void ImportMoveSubTrack( FbxAnimCurve* FbxCurve, int32 FbxDimension, UInterpTrackMoveAxis* SubTrack, int32 CurveIndex, bool bNegative, FbxAnimCurve* RealCurve, float DefaultVal );
+	void ImportMoveSubTrack( FbxAnimCurve* FbxCurve, int32 FbxDimension, class UInterpTrackMoveAxis* SubTrack, int32 CurveIndex, bool bNegative, FbxAnimCurve* RealCurve, float DefaultVal );
 
 	/**
 	 * Imports a FBX animated element into a Matinee track.
diff --git a/Engine/Source/Editor/UnrealEd/Private/Kismet2/BlueprintEditorUtils.cpp b/Engine/Source/Editor/UnrealEd/Private/Kismet2/BlueprintEditorUtils.cpp
--- a/Engine/Source/Editor/UnrealEd/Private/Kismet2/BlueprintEditorUtils.cpp
+++ b/Engine/Source/Editor/UnrealEd/Private/Kismet2/BlueprintEditorUtils.cpp
@@ -2,6 +2,8 @@
 
 #include "UnrealEd.h"
 
+#include "Matinee/MatineeActor.h"
+
 #include "LatentActions.h"
 #include "Kismet2/CompilerResultsLog.h"
 
diff --git a/Engine/Source/Editor/UnrealEd/Private/Lightmass/Lightmass.cpp b/Engine/Source/Editor/UnrealEd/Private/Lightmass/Lightmass.cpp
--- a/Engine/Source/Editor/UnrealEd/Private/Lightmass/Lightmass.cpp
+++ b/Engine/Source/Editor/UnrealEd/Private/Lightmass/Lightmass.cpp
@@ -8,6 +8,11 @@
 #include "PrecomputedLightVolume.h"
 #include "Runtime/Engine/Public/Landscape/LandscapeRender.h"
 #include "Runtime/Engine/Public/Landscape/LandscapeLight.h"
+#include "Runtime/Engine/Classes/Matinee/MatineeActor.h"
+#include "Runtime/Engine/Classes/Matinee/InterpGroup.h"
+#include "Runtime/Engine/Classes/Matinee/InterpGroupInst.h"
+#include "Runtime/Engine/Classes/Matinee/InterpTrackMove.h"
+#include "Runtime/Engine/Classes/Matinee/InterpTrackInstMove.h"
 
 extern FSwarmDebugOptions GSwarmDebugOptions;
 
diff --git a/Engine/Source/Editor/UnrealEd/Private/UnrealEdEngine.cpp b/Engine/Source/Editor/UnrealEd/Private/UnrealEdEngine.cpp
--- a/Engine/Source/Editor/UnrealEd/Private/UnrealEdEngine.cpp
+++ b/Engine/Source/Editor/UnrealEd/Private/UnrealEdEngine.cpp
@@ -1,6 +1,10 @@
 // Copyright 1998-2014 Epic Games, Inc. All Rights Reserved.
 
 #include "UnrealEd.h"
+
+#include "Matinee/InterpData.h"
+#include "Matinee/MatineeActor.h"
+
 #include "EditorSupportDelegates.h"
 #include "SoundDefinitions.h"
 #include "BusyCursor.h"
diff --git a/Engine/Source/Editor/UnrealEd/Private/UnrealEdSrv.cpp b/Engine/Source/Editor/UnrealEd/Private/UnrealEdSrv.cpp
--- a/Engine/Source/Editor/UnrealEd/Private/UnrealEdSrv.cpp
+++ b/Engine/Source/Editor/UnrealEd/Private/UnrealEdSrv.cpp
@@ -24,6 +24,7 @@
 #include "AssetSelection.h"
 #include "HighResScreenshot.h"
 #include "ActorEditorUtils.h"
+#include "Matinee/InterpData.h"
 
 DEFINE_LOG_CATEGORY_STATIC(LogUnrealEdSrv, Log, All);
 
diff --git a/Engine/Source/Editor/UnrealEd/Public/MatineeExporter.h b/Engine/Source/Editor/UnrealEd/Public/MatineeExporter.h
--- a/Engine/Source/Editor/UnrealEd/Public/MatineeExporter.h
+++ b/Engine/Source/Editor/UnrealEd/Public/MatineeExporter.h
@@ -7,6 +7,10 @@
 #ifndef __MATINEEEXPORTER_H__
 #define __MATINEEEXPORTER_H__
 
+#include "Matinee/MatineeActor.h"
+#include "Matinee/InterpGroup.h"
+#include "Matinee/InterpGroupInst.h"
+
 class UInterpData;
 class UInterpTrackMove;
 class UInterpTrackFloatProp;
diff --git a/Engine/Source/Editor/WorldBrowser/Private/LevelCollectionModel.cpp b/Engine/Source/Editor/WorldBrowser/Private/LevelCollectionModel.cpp
--- a/Engine/Source/Editor/WorldBrowser/Private/LevelCollectionModel.cpp
+++ b/Engine/Source/Editor/WorldBrowser/Private/LevelCollectionModel.cpp
@@ -4,6 +4,7 @@
 #include "SourceControlWindows.h"
 #include "AssetToolsModule.h"
 #include "EditorSupportDelegates.h"
+#include "Matinee/MatineeActor.h"
 
 #include "LevelCollectionModel.h"
 
diff --git a/Engine/Source/Runtime/Engine/Classes/Animation/SkeletalMeshActor.h b/Engine/Source/Runtime/Engine/Classes/Animation/SkeletalMeshActor.h
--- a/Engine/Source/Runtime/Engine/Classes/Animation/SkeletalMeshActor.h
+++ b/Engine/Source/Runtime/Engine/Classes/Animation/SkeletalMeshActor.h
@@ -1,6 +1,7 @@
 // Copyright 1998-2014 Epic Games, Inc. All Rights Reserved.
 #pragma once
 #include "Components/SkeletalMeshComponent.h"
+#include "Matinee/MatineeAnimInterface.h"
 #include "SkeletalMeshActor.generated.h"
 
 UCLASS(ClassGroup=ISkeletalMeshes, Blueprintable, ConversionRoot, meta=(ChildCanTick))
diff --git a/Engine/Source/Runtime/Engine/Classes/Matinee/InterpFilter_Classes.h b/Engine/Source/Runtime/Engine/Classes/Matinee/InterpFilter_Classes.h
--- a/Engine/Source/Runtime/Engine/Classes/Matinee/InterpFilter_Classes.h
+++ b/Engine/Source/Runtime/Engine/Classes/Matinee/InterpFilter_Classes.h
@@ -7,6 +7,7 @@
  */
 
 #pragma once
+#include "Matinee/InterpFilter.h"
 #include "InterpFilter_Classes.generated.h"
 
 UCLASS()
diff --git a/Engine/Source/Runtime/Engine/Classes/Matinee/InterpFilter_Custom.h b/Engine/Source/Runtime/Engine/Classes/Matinee/InterpFilter_Custom.h
--- a/Engine/Source/Runtime/Engine/Classes/Matinee/InterpFilter_Custom.h
+++ b/Engine/Source/Runtime/Engine/Classes/Matinee/InterpFilter_Custom.h
@@ -7,6 +7,7 @@
  */
 
 #pragma once
+#include "Matinee/InterpFilter.h"
 #include "InterpFilter_Custom.generated.h"
 
 UCLASS(MinimalAPI)
diff --git a/Engine/Source/Runtime/Engine/Classes/Matinee/InterpGroupCamera.h b/Engine/Source/Runtime/Engine/Classes/Matinee/InterpGroupCamera.h
--- a/Engine/Source/Runtime/Engine/Classes/Matinee/InterpGroupCamera.h
+++ b/Engine/Source/Runtime/Engine/Classes/Matinee/InterpGroupCamera.h
@@ -6,6 +6,7 @@
  */
 
 #pragma once
+#include "Matinee/InterpGroup.h"
 #include "InterpGroupCamera.generated.h"
 
 /**
diff --git a/Engine/Source/Runtime/Engine/Classes/Matinee/InterpGroupDirector.h b/Engine/Source/Runtime/Engine/Classes/Matinee/InterpGroupDirector.h
--- a/Engine/Source/Runtime/Engine/Classes/Matinee/InterpGroupDirector.h
+++ b/Engine/Source/Runtime/Engine/Classes/Matinee/InterpGroupDirector.h
@@ -6,6 +6,7 @@
  */
 
 #pragma once
+#include "Matinee/InterpGroup.h"
 #include "InterpGroupDirector.generated.h"
 
 UCLASS(collapsecategories, hidecategories=Object, MinimalAPI)
diff --git a/Engine/Source/Runtime/Engine/Classes/Matinee/InterpGroupInstDirector.h b/Engine/Source/Runtime/Engine/Classes/Matinee/InterpGroupInstDirector.h
--- a/Engine/Source/Runtime/Engine/Classes/Matinee/InterpGroupInstDirector.h
+++ b/Engine/Source/Runtime/Engine/Classes/Matinee/InterpGroupInstDirector.h
@@ -2,6 +2,7 @@
 
 
 #pragma once
+#include "Matinee/InterpGroupInst.h"
 #include "InterpGroupInstDirector.generated.h"
 
 UCLASS(MinimalAPI)
diff --git a/Engine/Source/Runtime/Engine/Classes/Matinee/InterpTrackAnimControl.h b/Engine/Source/Runtime/Engine/Classes/Matinee/InterpTrackAnimControl.h
--- a/Engine/Source/Runtime/Engine/Classes/Matinee/InterpTrackAnimControl.h
+++ b/Engine/Source/Runtime/Engine/Classes/Matinee/InterpTrackAnimControl.h
@@ -1,6 +1,7 @@
 // Copyright 1998-2014 Epic Games, Inc. All Rights Reserved.
 
 #pragma once
+#include "Matinee/InterpTrackFloatBase.h"
 #include "InterpTrackAnimControl.generated.h"
 
 /** Structure used for holding information for one animation played on the Anim Control track. */
@@ -83,7 +84,7 @@ class UInterpTrackAnimControl : public UInterpTrackFloatBase
 	virtual float GetTrackEndTime() const OVERRIDE;
 	virtual float GetKeyframeTime(int32 KeyIndex) const OVERRIDE;
 	virtual int32 GetKeyframeIndex( float KeyTime ) const OVERRIDE;
-	virtual int32 AddKeyframe(float Time, UInterpTrackInst* TrInst, EInterpCurveMode InitInterpMode) OVERRIDE;
+	virtual int32 AddKeyframe(float Time, class UInterpTrackInst* TrInst, EInterpCurveMode InitInterpMode) OVERRIDE;
 	virtual int32 SetKeyframeTime(int32 KeyIndex, float NewKeyTime, bool bUpdateOrder=true) OVERRIDE;
 	virtual void RemoveKeyframe(int32 KeyIndex) OVERRIDE;
 	virtual int32 DuplicateKeyframe(int32 KeyIndex, float NewKeyTime, UInterpTrack* ToTrack = NULL) OVERRIDE;
diff --git a/Engine/Source/Runtime/Engine/Classes/Matinee/InterpTrackAudioMaster.h b/Engine/Source/Runtime/Engine/Classes/Matinee/InterpTrackAudioMaster.h
--- a/Engine/Source/Runtime/Engine/Classes/Matinee/InterpTrackAudioMaster.h
+++ b/Engine/Source/Runtime/Engine/Classes/Matinee/InterpTrackAudioMaster.h
@@ -2,6 +2,7 @@
 
 
 #pragma once
+#include "Matinee/InterpTrackVectorBase.h"
 #include "InterpTrackAudioMaster.generated.h"
 
 UCLASS(MinimalAPI, meta=( DisplayName = "Audio Master Track" ) )
diff --git a/Engine/Source/Runtime/Engine/Classes/Matinee/InterpTrackColorProp.h b/Engine/Source/Runtime/Engine/Classes/Matinee/InterpTrackColorProp.h
--- a/Engine/Source/Runtime/Engine/Classes/Matinee/InterpTrackColorProp.h
+++ b/Engine/Source/Runtime/Engine/Classes/Matinee/InterpTrackColorProp.h
@@ -2,6 +2,7 @@
 
 
 #pragma once
+#include "Matinee/InterpTrackVectorBase.h"
 #include "InterpTrackColorProp.generated.h"
 
 UCLASS(MinimalAPI, meta=( DisplayName = "Color Property Track" ) )
diff --git a/Engine/Source/Runtime/Engine/Classes/Matinee/InterpTrackColorScale.h b/Engine/Source/Runtime/Engine/Classes/Matinee/InterpTrackColorScale.h
--- a/Engine/Source/Runtime/Engine/Classes/Matinee/InterpTrackColorScale.h
+++ b/Engine/Source/Runtime/Engine/Classes/Matinee/InterpTrackColorScale.h
@@ -6,6 +6,7 @@
  */
 
 #pragma once
+#include "Matinee/InterpTrackVectorBase.h"
 #include "InterpTrackColorScale.generated.h"
 
 UCLASS(meta=( DisplayName = "Color Scale Track" ) )
diff --git a/Engine/Source/Runtime/Engine/Classes/Matinee/InterpTrackFade.h b/Engine/Source/Runtime/Engine/Classes/Matinee/InterpTrackFade.h
--- a/Engine/Source/Runtime/Engine/Classes/Matinee/InterpTrackFade.h
+++ b/Engine/Source/Runtime/Engine/Classes/Matinee/InterpTrackFade.h
@@ -1,6 +1,7 @@
 // Copyright 1998-2014 Epic Games, Inc. All Rights Reserved.
 
 #pragma once
+#include "Matinee/InterpTrackFloatBase.h"
 #include "InterpTrackFade.generated.h"
 
 UCLASS(meta=( DisplayName = "Fade Track" ) )
diff --git a/Engine/Source/Runtime/Engine/Classes/Matinee/InterpTrackFloatBase.h b/Engine/Source/Runtime/Engine/Classes/Matinee/InterpTrackFloatBase.h
--- a/Engine/Source/Runtime/Engine/Classes/Matinee/InterpTrackFloatBase.h
+++ b/Engine/Source/Runtime/Engine/Classes/Matinee/InterpTrackFloatBase.h
@@ -1,6 +1,7 @@
 // Copyright 1998-2014 Epic Games, Inc. All Rights Reserved.
 
 #pragma once
+#include "Matinee/InterpTrack.h"
 #include "InterpTrackFloatBase.generated.h"
 
 UCLASS(abstract, MinimalAPI)
diff --git a/Engine/Source/Runtime/Engine/Classes/Matinee/InterpTrackFloatProp.h b/Engine/Source/Runtime/Engine/Classes/Matinee/InterpTrackFloatProp.h
--- a/Engine/Source/Runtime/Engine/Classes/Matinee/InterpTrackFloatProp.h
+++ b/Engine/Source/Runtime/Engine/Classes/Matinee/InterpTrackFloatProp.h
@@ -1,6 +1,7 @@
 // Copyright 1998-2014 Epic Games, Inc. All Rights Reserved.
 
 #pragma once
+#include "Matinee/InterpTrackFloatBase.h"
 #include "InterpTrackFloatProp.generated.h"
 
 UCLASS(MinimalAPI, meta=( DisplayName = "Float Property Track" ) )
diff --git a/Engine/Source/Runtime/Engine/Classes/Matinee/InterpTrackInst.h b/Engine/Source/Runtime/Engine/Classes/Matinee/InterpTrackInst.h
--- a/Engine/Source/Runtime/Engine/Classes/Matinee/InterpTrackInst.h
+++ b/Engine/Source/Runtime/Engine/Classes/Matinee/InterpTrackInst.h
@@ -9,6 +9,8 @@
 #pragma once
 #include "InterpTrackInst.generated.h"
 
+class UInterpTrack;
+
 UCLASS(MinimalAPI)
 class UInterpTrackInst : public UObject
 {
diff --git a/Engine/Source/Runtime/Engine/Classes/Matinee/InterpTrackInstDirector.h b/Engine/Source/Runtime/Engine/Classes/Matinee/InterpTrackInstDirector.h
--- a/Engine/Source/Runtime/Engine/Classes/Matinee/InterpTrackInstDirector.h
+++ b/Engine/Source/Runtime/Engine/Classes/Matinee/InterpTrackInstDirector.h
@@ -2,6 +2,7 @@
 
 
 #pragma once
+#include "Matinee/InterpTrackInst.h"
 #include "InterpTrackInstDirector.generated.h"
 
 UCLASS(MinimalAPI)
@@ -14,7 +15,7 @@ class UInterpTrackInstDirector : public UInterpTrackInst
 
 
 	// Begin UInterpTrackInst Instance
-	virtual void TermTrackInst(UInterpTrack* Track) OVERRIDE;
+	virtual void TermTrackInst(class UInterpTrack* Track) OVERRIDE;
 	// End UInterpTrackInst Instance
 };
 
diff --git a/Engine/Source/Runtime/Engine/Classes/Matinee/InterpTrackInstFloatProp.h b/Engine/Source/Runtime/Engine/Classes/Matinee/InterpTrackInstFloatProp.h
--- a/Engine/Source/Runtime/Engine/Classes/Matinee/InterpTrackInstFloatProp.h
+++ b/Engine/Source/Runtime/Engine/Classes/Matinee/InterpTrackInstFloatProp.h
@@ -2,6 +2,7 @@
 
 
 #pragma once
+#include "Matinee/InterpTrackInstProperty.h"
 #include "InterpTrackInstFloatProp.generated.h"
 
 UCLASS(MinimalAPI)
diff --git a/Engine/Source/Runtime/Engine/Classes/Matinee/InterpTrackInstMove.h b/Engine/Source/Runtime/Engine/Classes/Matinee/InterpTrackInstMove.h
--- a/Engine/Source/Runtime/Engine/Classes/Matinee/InterpTrackInstMove.h
+++ b/Engine/Source/Runtime/Engine/Classes/Matinee/InterpTrackInstMove.h
@@ -2,6 +2,7 @@
 
 
 #pragma once
+#include "Matinee/InterpTrackInst.h"
 #include "InterpTrackInstMove.generated.h"
 
 UCLASS(MinimalAPI)
diff --git a/Engine/Source/Runtime/Engine/Classes/Matinee/InterpTrackInstProperty.h b/Engine/Source/Runtime/Engine/Classes/Matinee/InterpTrackInstProperty.h
--- a/Engine/Source/Runtime/Engine/Classes/Matinee/InterpTrackInstProperty.h
+++ b/Engine/Source/Runtime/Engine/Classes/Matinee/InterpTrackInstProperty.h
@@ -2,6 +2,7 @@
 
 
 #pragma once
+#include "Matinee/InterpTrackInst.h"
 #include "InterpTrackInstProperty.generated.h"
 
 UCLASS()
diff --git a/Engine/Source/Runtime/Engine/Classes/Matinee/InterpTrackLinearColorProp.h b/Engine/Source/Runtime/Engine/Classes/Matinee/InterpTrackLinearColorProp.h
--- a/Engine/Source/Runtime/Engine/Classes/Matinee/InterpTrackLinearColorProp.h
+++ b/Engine/Source/Runtime/Engine/Classes/Matinee/InterpTrackLinearColorProp.h
@@ -2,6 +2,7 @@
 
 
 #pragma once
+#include "Matinee/InterpTrackLinearColorBase.h"
 #include "InterpTrackLinearColorProp.generated.h"
 
 UCLASS(MinimalAPI, meta=( DisplayName = "LinearColor Property Track" ) )
diff --git a/Engine/Source/Runtime/Engine/Classes/Matinee/InterpTrackMove.h b/Engine/Source/Runtime/Engine/Classes/Matinee/InterpTrackMove.h
--- a/Engine/Source/Runtime/Engine/Classes/Matinee/InterpTrackMove.h
+++ b/Engine/Source/Runtime/Engine/Classes/Matinee/InterpTrackMove.h
@@ -1,6 +1,7 @@
 // Copyright 1998-2014 Epic Games, Inc. All Rights Reserved.
 
 #pragma once
+#include "Matinee/InterpTrack.h"
 #include "InterpTrackMove.generated.h"
 
 UENUM()
diff --git a/Engine/Source/Runtime/Engine/Classes/Matinee/InterpTrackMoveAxis.h b/Engine/Source/Runtime/Engine/Classes/Matinee/InterpTrackMoveAxis.h
--- a/Engine/Source/Runtime/Engine/Classes/Matinee/InterpTrackMoveAxis.h
+++ b/Engine/Source/Runtime/Engine/Classes/Matinee/InterpTrackMoveAxis.h
@@ -1,6 +1,8 @@
 // Copyright 1998-2014 Epic Games, Inc. All Rights Reserved.
 
 #pragma once
+#include "Matinee/InterpTrackFloatBase.h"
+#include "Matinee/InterpTrackMove.h"
 #include "InterpTrackMoveAxis.generated.h"
 
 /**
diff --git a/Engine/Source/Runtime/Engine/Classes/Matinee/InterpTrackParticleReplay.h b/Engine/Source/Runtime/Engine/Classes/Matinee/InterpTrackParticleReplay.h
--- a/Engine/Source/Runtime/Engine/Classes/Matinee/InterpTrackParticleReplay.h
+++ b/Engine/Source/Runtime/Engine/Classes/Matinee/InterpTrackParticleReplay.h
@@ -1,64 +1,13 @@
 // Copyright 1998-2014 Epic Games, Inc. All Rights Reserved.
 
 #pragma once
+#include "Matinee/InterpTrack.h"
 #include "InterpTrackParticleReplay.generated.h"
 
 /**
- *
- * 
  *	This track implements support for creating and playing back captured particle system data
  */
 
-
-
-
-
-
-
-
-
-
-
-
-
-
-
-
-
-
-
-
-
-
-
-
-
-
-
-
-
-
-
-
-
-
-
-
-
-
-
-
-
-
-
-
-
-
-
-
-
-
-
 /** Data for a single key in this track */
 USTRUCT()
 struct FParticleReplayTrackKey
diff --git a/Engine/Source/Runtime/Engine/Classes/Matinee/InterpTrackSound.h b/Engine/Source/Runtime/Engine/Classes/Matinee/InterpTrackSound.h
--- a/Engine/Source/Runtime/Engine/Classes/Matinee/InterpTrackSound.h
+++ b/Engine/Source/Runtime/Engine/Classes/Matinee/InterpTrackSound.h
@@ -1,6 +1,7 @@
 // Copyright 1998-2014 Epic Games, Inc. All Rights Reserved.
 
 #pragma once
+#include "Matinee/InterpTrackVectorBase.h"
 #include "InterpTrackSound.generated.h"
 
 /**
diff --git a/Engine/Source/Runtime/Engine/Classes/Matinee/InterpTrackVectorProp.h b/Engine/Source/Runtime/Engine/Classes/Matinee/InterpTrackVectorProp.h
--- a/Engine/Source/Runtime/Engine/Classes/Matinee/InterpTrackVectorProp.h
+++ b/Engine/Source/Runtime/Engine/Classes/Matinee/InterpTrackVectorProp.h
@@ -2,6 +2,7 @@
 
 
 #pragma once
+#include "Matinee/InterpTrackVectorBase.h"
 #include "InterpTrackVectorProp.generated.h"
 
 UCLASS(MinimalAPI, meta=( DisplayName = "Vector Property Track" ) )
diff --git a/Engine/Source/Runtime/Engine/Classes/Matinee/MatineeActorCameraAnim.h b/Engine/Source/Runtime/Engine/Classes/Matinee/MatineeActorCameraAnim.h
--- a/Engine/Source/Runtime/Engine/Classes/Matinee/MatineeActorCameraAnim.h
+++ b/Engine/Source/Runtime/Engine/Classes/Matinee/MatineeActorCameraAnim.h
@@ -1,7 +1,7 @@
 // Copyright 1998-2014 Epic Games, Inc. All Rights Reserved.
 
 #pragma once
-
+#include "Matinee/MatineeActor.h"
 #include "MatineeActorCameraAnim.generated.h"
 
 /**
diff --git a/Engine/Source/Runtime/Engine/Private/Actor.cpp b/Engine/Source/Runtime/Engine/Private/Actor.cpp
--- a/Engine/Source/Runtime/Engine/Private/Actor.cpp
+++ b/Engine/Source/Runtime/Engine/Private/Actor.cpp
@@ -8,6 +8,9 @@
 #include "MessageLog.h"
 #include "Net/UnrealNetwork.h"
 #include "DisplayDebugHelpers.h"
+#include "Matinee/MatineeActor.h"
+#include "Matinee/InterpGroup.h"
+#include "Matinee/InterpGroupInst.h"
 
 //DEFINE_LOG_CATEGORY_STATIC(LogActor, Log, All);
 DEFINE_LOG_CATEGORY(LogActor);
diff --git a/Engine/Source/Runtime/Engine/Private/Camera/CameraAnim.cpp b/Engine/Source/Runtime/Engine/Private/Camera/CameraAnim.cpp
--- a/Engine/Source/Runtime/Engine/Private/Camera/CameraAnim.cpp
+++ b/Engine/Source/Runtime/Engine/Private/Camera/CameraAnim.cpp
@@ -2,6 +2,13 @@
 
 #include "EnginePrivate.h"
 #include "ParticleDefinitions.h"
+#include "Matinee/MatineeActor.h"
+#include "Matinee/InterpData.h"
+#include "Matinee/InterpGroup.h"
+#include "Matinee/InterpGroupInst.h"
+#include "Matinee/InterpGroupCamera.h"
+#include "Matinee/InterpTrackMove.h"
+
 
 //////////////////////////////////////////////////////////////////////////
 // UCameraAnim
diff --git a/Engine/Source/Runtime/Engine/Private/Camera/CameraAnimInst.cpp b/Engine/Source/Runtime/Engine/Private/Camera/CameraAnimInst.cpp
--- a/Engine/Source/Runtime/Engine/Private/Camera/CameraAnimInst.cpp
+++ b/Engine/Source/Runtime/Engine/Private/Camera/CameraAnimInst.cpp
@@ -2,6 +2,11 @@
 
 
 #include "EnginePrivate.h"
+#include "Matinee/InterpTrackMove.h"
+#include "Matinee/InterpGroup.h"
+#include "Matinee/InterpGroupInst.h"
+#include "Matinee/InterpTrackInstMove.h"
+#include "Matinee/InterpTrackFloatProp.h"
 
 //////////////////////////////////////////////////////////////////////////
 // UCameraAnimInst
diff --git a/Engine/Source/Runtime/Engine/Private/GameMode.cpp b/Engine/Source/Runtime/Engine/Private/GameMode.cpp
--- a/Engine/Source/Runtime/Engine/Private/GameMode.cpp
+++ b/Engine/Source/Runtime/Engine/Private/GameMode.cpp
@@ -5,6 +5,7 @@
 =============================================================================*/
 
 #include "EnginePrivate.h"
+#include "Matinee/MatineeActor.h"
 #include "OnlineSubsystemUtils.h"
 
 DEFINE_LOG_CATEGORY_STATIC(LogGameMode, Log, All);
diff --git a/Engine/Source/Runtime/Engine/Private/InterpCurveEdSetup.cpp b/Engine/Source/Runtime/Engine/Private/InterpCurveEdSetup.cpp
--- a/Engine/Source/Runtime/Engine/Private/InterpCurveEdSetup.cpp
+++ b/Engine/Source/Runtime/Engine/Private/InterpCurveEdSetup.cpp
@@ -5,6 +5,7 @@
 =============================================================================*/
 
 #include "EnginePrivate.h"
+#include "Matinee/InterpTrack.h"
 
 
 void UInterpCurveEdSetup::PostLoad()
diff --git a/Engine/Source/Runtime/Engine/Private/Interpolation.cpp b/Engine/Source/Runtime/Engine/Private/Interpolation.cpp
--- a/Engine/Source/Runtime/Engine/Private/Interpolation.cpp
+++ b/Engine/Source/Runtime/Engine/Private/Interpolation.cpp
@@ -5,6 +5,59 @@
 =============================================================================*/
 
 #include "EnginePrivate.h"
+#include "Matinee/MatineeActorCameraAnim.h"
+#include "Matinee/InterpData.h"
+#include "Matinee/InterpTrackInstProperty.h"
+#include "Matinee/InterpTrackMove.h"
+#include "Matinee/InterpTrackMoveAxis.h"
+#include "Matinee/InterpTrackInstMove.h"
+#include "Matinee/InterpTrackDirector.h"
+#include "Matinee/InterpTrackInstDirector.h"
+#include "Matinee/InterpTrackVisibility.h"
+#include "Matinee/InterpTrackInstVisibility.h"
+#include "Matinee/InterpTrackAnimControl.h"
+#include "Matinee/InterpTrackInstAnimControl.h"
+#include "Matinee/InterpTrackEvent.h"
+#include "Matinee/InterpTrackInstEvent.h"
+#include "Matinee/InterpTrackToggle.h"
+#include "Matinee/InterpTrackInstToggle.h"
+#include "Matinee/InterpTrackFade.h"
+#include "Matinee/InterpTrackInstFade.h"
+#include "Matinee/InterpTrackSlomo.h"
+#include "Matinee/InterpTrackInstSlomo.h"
+#include "Matinee/InterpTrackSound.h"
+#include "Matinee/InterpTrackInstSound.h"
+#include "Matinee/InterpTrackLinearColorBase.h"
+#include "Matinee/InterpTrackFloatProp.h"
+#include "Matinee/InterpTrackInstFloatProp.h"
+#include "Matinee/InterpTrackVectorProp.h"
+#include "Matinee/InterpTrackInstVectorProp.h"
+#include "Matinee/InterpTrackBoolProp.h"
+#include "Matinee/InterpTrackInstBoolProp.h"
+#include "Matinee/InterpTrackColorProp.h"
+#include "Matinee/InterpTrackInstColorProp.h"
+#include "Matinee/InterpTrackLinearColorProp.h"
+#include "Matinee/InterpTrackInstLinearColorProp.h"
+#include "Matinee/InterpTrackAudioMaster.h"
+#include "Matinee/InterpTrackInstAudioMaster.h"
+#include "Matinee/InterpTrackColorScale.h"
+#include "Matinee/InterpTrackInstColorScale.h"
+#include "Matinee/InterpTrackFloatParticleParam.h"
+#include "Matinee/InterpTrackInstFloatParticleParam.h"
+#include "Matinee/InterpTrackFloatMaterialParam.h"
+#include "Matinee/InterpTrackInstFloatMaterialParam.h"
+#include "Matinee/InterpTrackVectorMaterialParam.h"
+#include "Matinee/InterpTrackInstVectorMaterialParam.h"
+#include "Matinee/InterpTrackParticleReplay.h"
+#include "Matinee/InterpTrackInstParticleReplay.h"
+#include "Matinee/InterpGroup.h"
+#include "Matinee/InterpGroupInst.h"
+#include "Matinee/InterpGroupDirector.h"
+#include "Matinee/InterpGroupInstDirector.h"
+#include "Matinee/InterpGroupCamera.h"
+#include "Matinee/InterpGroupInstCamera.h"
+#include "Matinee/InterpFilter_Classes.h"
+#include "Matinee/InterpFilter_Custom.h"
 #include "ParticleDefinitions.h"
 #include "SoundDefinitions.h"
 #include "InterpolationHitProxy.h"
diff --git a/Engine/Source/Runtime/Engine/Private/InterpolationCurveEd.cpp b/Engine/Source/Runtime/Engine/Private/InterpolationCurveEd.cpp
--- a/Engine/Source/Runtime/Engine/Private/InterpolationCurveEd.cpp
+++ b/Engine/Source/Runtime/Engine/Private/InterpolationCurveEd.cpp
@@ -5,6 +5,10 @@
 =============================================================================*/
 
 #include "EnginePrivate.h"
+#include "Matinee/InterpTrackMove.h"
+#include "Matinee/InterpTrackMoveAxis.h"
+#include "Matinee/InterpTrackVectorBase.h"
+#include "Matinee/InterpTrackLinearColorBase.h"
 
 /*-----------------------------------------------------------------------------
  UInterpTrackMove
diff --git a/Engine/Source/Runtime/Engine/Private/InterpolationDraw.cpp b/Engine/Source/Runtime/Engine/Private/InterpolationDraw.cpp
--- a/Engine/Source/Runtime/Engine/Private/InterpolationDraw.cpp
+++ b/Engine/Source/Runtime/Engine/Private/InterpolationDraw.cpp
@@ -5,6 +5,29 @@
 =============================================================================*/
 
 #include "EnginePrivate.h"
+#include "Matinee/MatineeActor.h"
+#include "Matinee/InterpData.h"
+#include "Matinee/InterpGroup.h"
+#include "Matinee/InterpGroupInst.h"
+#include "Matinee/InterpTrackMove.h"
+#include "Matinee/InterpTrackInstMove.h"
+#include "Matinee/InterpTrackMoveAxis.h"
+#include "Matinee/InterpTrackToggle.h"
+#include "Matinee/InterpTrackEvent.h"
+#include "Matinee/InterpTrackFade.h"
+#include "Matinee/InterpTrackDirector.h"
+#include "Matinee/InterpTrackAnimControl.h"
+#include "Matinee/InterpTrackFloatProp.h"
+#include "Matinee/InterpTrackBoolProp.h"
+#include "Matinee/InterpTrackVectorProp.h"
+#include "Matinee/InterpTrackLinearColorProp.h"
+#include "Matinee/InterpTrackColorProp.h"
+#include "Matinee/InterpTrackSound.h"
+#include "Matinee/InterpTrackSlomo.h"
+#include "Matinee/InterpTrackColorScale.h"
+#include "Matinee/InterpTrackAudioMaster.h"
+#include "Matinee/InterpTrackVisibility.h"
+#include "Matinee/InterpTrackParticleReplay.h"
 #include "SoundDefinitions.h"
 #include "InterpolationHitProxy.h"
 #include "AnimationUtils.h"
diff --git a/Engine/Source/Runtime/Engine/Private/LocalPlayer.cpp b/Engine/Source/Runtime/Engine/Private/LocalPlayer.cpp
--- a/Engine/Source/Runtime/Engine/Private/LocalPlayer.cpp
+++ b/Engine/Source/Runtime/Engine/Private/LocalPlayer.cpp
@@ -2,6 +2,9 @@
 
 #include "EnginePrivate.h"
 
+#include "Matinee/MatineeActor.h"
+#include "Matinee/InterpData.h"
+#include "Matinee/InterpGroupInst.h"
 #include "SubtitleManager.h"
 #include "Net/UnrealNetwork.h"
 #include "OnlineSubsystemUtils.h"
diff --git a/Engine/Source/Runtime/Engine/Private/MatineeInterface.cpp b/Engine/Source/Runtime/Engine/Private/MatineeInterface.cpp
--- a/Engine/Source/Runtime/Engine/Private/MatineeInterface.cpp
+++ b/Engine/Source/Runtime/Engine/Private/MatineeInterface.cpp
@@ -1,6 +1,7 @@
 // Copyright 1998-2014 Epic Games, Inc. All Rights Reserved.
 
 #include "EnginePrivate.h"
+#include "Matinee/MatineeInterface.h"
 
 UMatineeInterface::UMatineeInterface(const class FPostConstructInitializeProperties& PCIP)
 	: Super(PCIP)
diff --git a/Engine/Source/Runtime/Engine/Private/PlayerController.cpp b/Engine/Source/Runtime/Engine/Private/PlayerController.cpp
--- a/Engine/Source/Runtime/Engine/Private/PlayerController.cpp
+++ b/Engine/Source/Runtime/Engine/Private/PlayerController.cpp
@@ -10,6 +10,8 @@
 #include "Slate.h"
 #include "GameFramework/TouchInterface.h"
 #include "DisplayDebugHelpers.h"
+#include "Matinee/InterpTrackInstDirector.h"
+#include "Matinee/MatineeActor.h"
 
 DEFINE_LOG_CATEGORY_STATIC(LogPlayerController, Log, All);
 
diff --git a/Engine/Source/Runtime/Engine/Private/Replication.cpp b/Engine/Source/Runtime/Engine/Private/Replication.cpp
--- a/Engine/Source/Runtime/Engine/Private/Replication.cpp
+++ b/Engine/Source/Runtime/Engine/Private/Replication.cpp
@@ -5,6 +5,8 @@
 =============================================================================*/
 
 #include "EnginePrivate.h"
+#include "Matinee/MatineeInterface.h"
+#include "Matinee/MatineeActor.h"
 #include "Net/UnrealNetwork.h"
 #include "NavigationPathBuilder.h"
 #include "SoundDefinitions.h"
diff --git a/Engine/Source/Runtime/Engine/Private/Tests/AutomationTestCommon.cpp b/Engine/Source/Runtime/Engine/Private/Tests/AutomationTestCommon.cpp
--- a/Engine/Source/Runtime/Engine/Private/Tests/AutomationTestCommon.cpp
+++ b/Engine/Source/Runtime/Engine/Private/Tests/AutomationTestCommon.cpp
@@ -1,6 +1,7 @@
 // Copyright 1998-2014 Epic Games, Inc. All Rights Reserved.
 
 #include "EnginePrivate.h"
+#include "Matinee/MatineeActor.h"
 #include "AutomationTestCommon.h"
 
 
diff --git a/Engine/Source/Runtime/Engine/Private/UnrealClient.cpp b/Engine/Source/Runtime/Engine/Private/UnrealClient.cpp
--- a/Engine/Source/Runtime/Engine/Private/UnrealClient.cpp
+++ b/Engine/Source/Runtime/Engine/Private/UnrealClient.cpp
@@ -2,6 +2,7 @@
 
 
 #include "EnginePrivate.h"
+#include "Matinee/MatineeActor.h"
 #include "EditorSupportDelegates.h"
 #include "RenderCore.h"
 #include "RenderResource.h"
diff --git a/Engine/Source/Runtime/Engine/Private/World.cpp b/Engine/Source/Runtime/Engine/Private/World.cpp
--- a/Engine/Source/Runtime/Engine/Private/World.cpp
+++ b/Engine/Source/Runtime/Engine/Private/World.cpp
@@ -5,6 +5,7 @@
 =============================================================================*/
 
 #include "EnginePrivate.h"
+#include "Matinee/MatineeActor.h"
 #include "Net/UnrealNetwork.h"
 #include "ParticleDefinitions.h"
 #include "Database.h"
diff --git a/Engine/Source/Runtime/Engine/Public/Engine.h b/Engine/Source/Runtime/Engine/Public/Engine.h
--- a/Engine/Source/Runtime/Engine/Public/Engine.h
+++ b/Engine/Source/Runtime/Engine/Public/Engine.h
@@ -124,8 +124,6 @@
 #include "Engine/SpotLight.h"
 #include "Engine/GeneratedMeshAreaLight.h"
 #include "Materials/MaterialInstanceActor.h"
-#include "Matinee/MatineeActor.h"
-#include "Matinee/MatineeActorCameraAnim.h"
 #include "AI/Navigation/NavigationData.h"
 #include "AI/Navigation/NavigationGraph.h"
 #include "AI/Navigation/NavAreas/NavArea.h"
@@ -174,7 +172,6 @@
 #include "PhysicsEngine/ConstraintInstance.h"
 #include "PhysicsEngine/PhysicsConstraintComponent.h"
 #include "Components/SkeletalMeshComponent.h"
-#include "Matinee/MatineeAnimInterface.h"
 #include "Animation/SkeletalMeshActor.h"
 #include "Engine/StaticMeshActor.h"
 #include "Engine/InteractiveFoliageActor.h"
@@ -470,63 +467,6 @@
 #include "GameFramework/InputSettings.h"
 #include "Foliage/InstancedFoliageSettings.h"
 #include "Engine/InterpCurveEdSetup.h"
-#include "Matinee/InterpData.h"
-#include "Matinee/InterpFilter.h"
-#include "Matinee/InterpFilter_Classes.h"
-#include "Matinee/InterpFilter_Custom.h"
-#include "Matinee/InterpGroup.h"
-#include "Matinee/InterpGroupCamera.h"
-#include "Matinee/InterpGroupDirector.h"
-#include "Matinee/InterpGroupInst.h"
-#include "Matinee/InterpGroupInstCamera.h"
-#include "Matinee/InterpGroupInstDirector.h"
-#include "Matinee/InterpTrack.h"
-#include "Matinee/InterpTrackBoolProp.h"
-#include "Matinee/InterpTrackDirector.h"
-#include "Matinee/InterpTrackEvent.h"
-#include "Matinee/InterpTrackFloatBase.h"
-#include "Matinee/InterpTrackAnimControl.h"
-#include "Matinee/InterpTrackFade.h"
-#include "Matinee/InterpTrackFloatMaterialParam.h"
-#include "Matinee/InterpTrackFloatParticleParam.h"
-#include "Matinee/InterpTrackFloatProp.h"
-#include "Matinee/InterpTrackMove.h"
-#include "Matinee/InterpTrackMoveAxis.h"
-#include "Matinee/InterpTrackSlomo.h"
-#include "Matinee/InterpTrackLinearColorBase.h"
-#include "Matinee/InterpTrackLinearColorProp.h"
-#include "Matinee/InterpTrackParticleReplay.h"
-#include "Matinee/InterpTrackToggle.h"
-#include "Matinee/InterpTrackVectorBase.h"
-#include "Matinee/InterpTrackAudioMaster.h"
-#include "Matinee/InterpTrackColorProp.h"
-#include "Matinee/InterpTrackColorScale.h"
-#include "Matinee/InterpTrackSound.h"
-#include "Matinee/InterpTrackVectorMaterialParam.h"
-#include "Matinee/InterpTrackVectorProp.h"
-#include "Matinee/InterpTrackVisibility.h"
-#include "Matinee/InterpTrackInst.h"
-#include "Matinee/InterpTrackInstAnimControl.h"
-#include "Matinee/InterpTrackInstAudioMaster.h"
-#include "Matinee/InterpTrackInstColorScale.h"
-#include "Matinee/InterpTrackInstDirector.h"
-#include "Matinee/InterpTrackInstEvent.h"
-#include "Matinee/InterpTrackInstFade.h"
-#include "Matinee/InterpTrackInstFloatMaterialParam.h"
-#include "Matinee/InterpTrackInstFloatParticleParam.h"
-#include "Matinee/InterpTrackInstMove.h"
-#include "Matinee/InterpTrackInstParticleReplay.h"
-#include "Matinee/InterpTrackInstProperty.h"
-#include "Matinee/InterpTrackInstBoolProp.h"
-#include "Matinee/InterpTrackInstColorProp.h"
-#include "Matinee/InterpTrackInstFloatProp.h"
-#include "Matinee/InterpTrackInstLinearColorProp.h"
-#include "Matinee/InterpTrackInstVectorProp.h"
-#include "Matinee/InterpTrackInstSlomo.h"
-#include "Matinee/InterpTrackInstSound.h"
-#include "Matinee/InterpTrackInstToggle.h"
-#include "Matinee/InterpTrackInstVectorMaterialParam.h"
-#include "Matinee/InterpTrackInstVisibility.h"
 #include "Engine/IntSerialization.h"
 #include "AI/KismetAIAsyncTaskProxy.h"
 #include "Landscape/LandscapeSplineSegment.h"
@@ -679,7 +619,6 @@
 #include "Materials/MaterialInstanceDynamic.h"
 #include "Materials/MaterialParameterCollection.h"
 #include "Materials/MaterialParameterCollectionInstance.h"
-#include "Matinee/MatineeInterface.h"
 #include "AI/Navigation/NavAreas/NavArea_Default.h"
 #include "AI/Navigation/NavAreas/NavArea_Null.h"
 #include "AI/Navigation/NavAreas/NavAreaMeta.h"
