Compiler warning fixes from PR 185.

This is a subset of changes (compiler warning fixes needed to compile Linux editor) from pull request 185 by sbc100, 3dluvr, wshearn, abarbu et al: https://github.com/EpicGames/UnrealEngine/pull/185
- Remove superfluous expressions without side effect.

diff --git a/Engine/Source/Editor/Kismet/Private/BlueprintEditorViewportContextMenuExtender.cpp b/Engine/Source/Editor/Kismet/Private/BlueprintEditorViewportContextMenuExtender.cpp
--- a/Engine/Source/Editor/Kismet/Private/BlueprintEditorViewportContextMenuExtender.cpp
+++ b/Engine/Source/Editor/Kismet/Private/BlueprintEditorViewportContextMenuExtender.cpp
@@ -73,7 +73,7 @@ void RecompileOutOfDateKismetForSelection()
 
 	// Run thru all selected actors, looking for out of date blueprints
 	FSelectionIterator SelectedActorItr( GEditor->GetSelectedActorIterator() );
-	for (SelectedActorItr; SelectedActorItr; ++SelectedActorItr)
+	for ( ; SelectedActorItr; ++SelectedActorItr)
 	{
 		AActor* CurrentActor = Cast<AActor>(*SelectedActorItr);
 
diff --git a/Engine/Source/Editor/Matinee/Private/MatineeTools.cpp b/Engine/Source/Editor/Matinee/Private/MatineeTools.cpp
--- a/Engine/Source/Editor/Matinee/Private/MatineeTools.cpp
+++ b/Engine/Source/Editor/Matinee/Private/MatineeTools.cpp
@@ -1395,7 +1395,7 @@ void FMatinee::ViewEndOfTrack()
 	if( GetSelectedTrackCount() > 0 )
 	{
 		FSelectedTrackIterator TrackIt(GetSelectedTrackIterator());
-		for( TrackIt; TrackIt; ++TrackIt )
+		for( ; TrackIt; ++TrackIt )
 		{
 			UInterpTrack* Track = *TrackIt;
 			if (Track->GetTrackEndTime() > NewEndTime)
@@ -2121,7 +2121,7 @@ void FMatinee::AddKey()
 	{
 		// Populate the list of tracks that we need to add keys to.
 		FSelectedTrackIterator TrackIt(GetSelectedTrackIterator());
-		for( TrackIt; TrackIt; ++TrackIt )
+		for( ; TrackIt; ++TrackIt )
 		{
 			// Only allow keys to be added to multiple tracks at once if they are subtracks of a movement track.
 			if( (*TrackIt)->IsA( UInterpTrackMoveAxis::StaticClass() ) ) 
diff --git a/Engine/Source/Editor/SceneOutliner/Private/SSceneOutliner.cpp b/Engine/Source/Editor/SceneOutliner/Private/SSceneOutliner.cpp
--- a/Engine/Source/Editor/SceneOutliner/Private/SSceneOutliner.cpp
+++ b/Engine/Source/Editor/SceneOutliner/Private/SSceneOutliner.cpp
@@ -591,7 +591,6 @@ namespace SceneOutliner
 			HideTemporaryActorsFilter = SceneOutlinerFilters::CreateHideTemporaryActorsFilter();
 		}
 
-		sizeof(TSharedPtr<TDelegateFilter<const AActor*const>>);
 		if ( bHideTemporaryActors )
 		{			
 			Filters->Add( HideTemporaryActorsFilter );
diff --git a/Engine/Source/Editor/UnrealEd/Private/Commandlets/GatherTextFromAssetsCommandlet.cpp b/Engine/Source/Editor/UnrealEd/Private/Commandlets/GatherTextFromAssetsCommandlet.cpp
--- a/Engine/Source/Editor/UnrealEd/Private/Commandlets/GatherTextFromAssetsCommandlet.cpp
+++ b/Engine/Source/Editor/UnrealEd/Private/Commandlets/GatherTextFromAssetsCommandlet.cpp
@@ -801,7 +801,7 @@ int32 UGatherTextFromAssetsCommandlet::Main(const FString& Params)
 	for( int32 BatchIndex = 0; BatchIndex < BatchCount; ++BatchIndex )
 	{
 		int32 PackagesInThisBatch = 0;
-		for( PackageIndex; PackageIndex < PackageCount && PackagesInThisBatch < PackagesPerBatchCount; ++PackageIndex )
+		for( ; PackageIndex < PackageCount && PackagesInThisBatch < PackagesPerBatchCount; ++PackageIndex )
 		{
 			FString PackageFileName = PackageFileNamesToLoad[PackageIndex];
 
diff --git a/Engine/Source/Editor/UnrealEd/Private/ConsolidateWindow.cpp b/Engine/Source/Editor/UnrealEd/Private/ConsolidateWindow.cpp
--- a/Engine/Source/Editor/UnrealEd/Private/ConsolidateWindow.cpp
+++ b/Engine/Source/Editor/UnrealEd/Private/ConsolidateWindow.cpp
@@ -665,7 +665,7 @@ bool SConsolidateToolWidget::AreObjClassesHomogeneous()
 		// Starting from the second consolidation object, iterate through all consolidation objects
 		// to see if they all share a common class
 		++ConsolidationObjIter;
-		for ( ConsolidationObjIter; ConsolidationObjIter; ++ConsolidationObjIter )
+		for ( ; ConsolidationObjIter; ++ConsolidationObjIter )
 		{
 			const UObject* CurObj = *ConsolidationObjIter;
 			check( CurObj );
diff --git a/Engine/Source/Editor/UnrealEd/Private/Editor.cpp b/Engine/Source/Editor/UnrealEd/Private/Editor.cpp
--- a/Engine/Source/Editor/UnrealEd/Private/Editor.cpp
+++ b/Engine/Source/Editor/UnrealEd/Private/Editor.cpp
@@ -3212,7 +3212,7 @@ void UEditorEngine::SelectAllActorsWithClass( bool bArchetype )
 			++SelectedActorIter;
 		}
 		// Check all the other selected actors
-		for ( SelectedActorIter; SelectedActorIter && bAllSameClassAndArchetype; ++SelectedActorIter )
+		for ( ; SelectedActorIter && bAllSameClassAndArchetype; ++SelectedActorIter )
 		{
 			AActor* CurActor = *SelectedActorIter;
 			if ( CurActor->GetClass() != FirstClass || CurActor->GetArchetype() != FirstArchetype )
diff --git a/Engine/Source/Editor/UnrealEd/Private/EditorBsp.cpp b/Engine/Source/Editor/UnrealEd/Private/EditorBsp.cpp
--- a/Engine/Source/Editor/UnrealEd/Private/EditorBsp.cpp
+++ b/Engine/Source/Editor/UnrealEd/Private/EditorBsp.cpp
@@ -1545,7 +1545,7 @@ int DistributePoint
 	{
 		// This point is coplanar with this node, so check point for intersection with
 		// this node's sides, then loop with its coplanars.
-		for( iNode; iNode!=INDEX_NONE; iNode=Model->Nodes[iNode].iPlane )
+		for( ; iNode!=INDEX_NONE; iNode=Model->Nodes[iNode].iPlane )
 		{
 			FVert* VertPool = &Model->Verts[Model->Nodes[iNode].iVertPool];
 
