UE4 Refactoring. Changed OVERRIDE and FINAL macros to keywords override and final.

diff --git a/Engine/Source/Editor/MaterialEditor/Private/SMaterialEditorViewport.h b/Engine/Source/Editor/MaterialEditor/Private/SMaterialEditorViewport.h
--- a/Engine/Source/Editor/MaterialEditor/Private/SMaterialEditorViewport.h
+++ b/Engine/Source/Editor/MaterialEditor/Private/SMaterialEditorViewport.h
@@ -19,9 +19,9 @@ public:
 	void Construct(const FArguments& InArgs);
 	~SMaterialEditorViewport();
 	
-	virtual void AddReferencedObjects( FReferenceCollector& Collector ) OVERRIDE;
+	virtual void AddReferencedObjects( FReferenceCollector& Collector ) override;
 
-	virtual void Tick(const FGeometry& AllottedGeometry, const double InCurrentTime, const float InDeltaTime) OVERRIDE;
+	virtual void Tick(const FGeometry& AllottedGeometry, const double InCurrentTime, const float InDeltaTime) override;
 	void RefreshViewport();
 	
 	/**
@@ -83,11 +83,11 @@ public:
 
 protected:
 	/** SEditorViewport interface */
-	virtual TSharedRef<FEditorViewportClient> MakeEditorViewportClient() OVERRIDE;
-	virtual TSharedPtr<SWidget> MakeViewportToolbar() OVERRIDE;
-	virtual EVisibility OnGetViewportContentVisibility() const OVERRIDE;
-	virtual void BindCommands() OVERRIDE;
-	virtual void OnFocusViewportToSelection() OVERRIDE;
+	virtual TSharedRef<FEditorViewportClient> MakeEditorViewportClient() override;
+	virtual TSharedPtr<SWidget> MakeViewportToolbar() override;
+	virtual EVisibility OnGetViewportContentVisibility() const override;
+	virtual void BindCommands() override;
+	virtual void OnFocusViewportToSelection() override;
 private:
 	/** The parent tab where this viewport resides */
 	TWeakPtr<SDockTab> ParentTab;
diff --git a/Engine/Source/Editor/MaterialEditor/Private/SMaterialPalette.h b/Engine/Source/Editor/MaterialEditor/Private/SMaterialPalette.h
--- a/Engine/Source/Editor/MaterialEditor/Private/SMaterialPalette.h
+++ b/Engine/Source/Editor/MaterialEditor/Private/SMaterialPalette.h
@@ -31,8 +31,8 @@ public:
 
 protected:
 	// SGraphPalette Interface
-	virtual TSharedRef<SWidget> OnCreateWidgetForAction(FCreateWidgetForActionData* const InCreateData) OVERRIDE;
-	virtual void CollectAllActions(FGraphActionListBuilderBase& OutAllActions) OVERRIDE;
+	virtual TSharedRef<SWidget> OnCreateWidgetForAction(FCreateWidgetForActionData* const InCreateData) override;
+	virtual void CollectAllActions(FGraphActionListBuilderBase& OutAllActions) override;
 	// End of SGraphPalette Interface
 
 	/** Get the currently selected category name */
diff --git a/Engine/Source/Editor/MaterialEditor/Public/MaterialEditorActions.h b/Engine/Source/Editor/MaterialEditor/Public/MaterialEditorActions.h
--- a/Engine/Source/Editor/MaterialEditor/Public/MaterialEditorActions.h
+++ b/Engine/Source/Editor/MaterialEditor/Public/MaterialEditorActions.h
@@ -158,7 +158,7 @@ public:
 	/**
 	 * Initialize commands
 	 */
-	virtual void RegisterCommands() OVERRIDE;
+	virtual void RegisterCommands() override;
 
 public:
 };
@@ -205,7 +205,7 @@ public:
 	}	
 
 	/** TCommands interface */
-	virtual void RegisterCommands() OVERRIDE;
+	virtual void RegisterCommands() override;
 
 	/**
 	 * Returns a graph action assigned to the passed in gesture
diff --git a/Engine/Source/Editor/Matinee/Private/Matinee.h b/Engine/Source/Editor/Matinee/Private/Matinee.h
--- a/Engine/Source/Editor/Matinee/Private/Matinee.h
+++ b/Engine/Source/Editor/Matinee/Private/Matinee.h
@@ -31,15 +31,15 @@ public:
 	void InitMatinee(const EToolkitMode::Type Mode, const TSharedPtr< class IToolkitHost >& InitToolkitHost, UObject* InObjectToEdit);
 
 	/** IToolkit interface */
-	virtual FName GetToolkitFName() const OVERRIDE;
-	virtual FText GetBaseToolkitName() const OVERRIDE;
-	virtual FString GetWorldCentricTabPrefix() const OVERRIDE;
-	virtual FLinearColor GetWorldCentricTabColorScale() const OVERRIDE;
+	virtual FName GetToolkitFName() const override;
+	virtual FText GetBaseToolkitName() const override;
+	virtual FString GetWorldCentricTabPrefix() const override;
+	virtual FLinearColor GetWorldCentricTabColorScale() const override;
 
 	/** Register tabs that this toolkit can spawn with the TabManager */
-	virtual void RegisterTabSpawners(const TSharedRef<class FTabManager>& TabManager) OVERRIDE;
+	virtual void RegisterTabSpawners(const TSharedRef<class FTabManager>& TabManager) override;
 	/** Unregister tabs that this toolkit can spawn */	
-	virtual void UnregisterTabSpawners(const TSharedRef<class FTabManager>& TabManager) OVERRIDE;
+	virtual void UnregisterTabSpawners(const TSharedRef<class FTabManager>& TabManager) override;
 
 	TSharedRef<SDockTab> SpawnTab( const FSpawnTabArgs& TabSpawnArgs, FName TabIdentifier );
 
@@ -49,7 +49,7 @@ public:
 	void ExtendDefaultToolbarMenu();
 
 	/** @return the documentation location for this editor */
-	virtual FString GetDocumentationLink() const OVERRIDE
+	virtual FString GetDocumentationLink() const override
 	{
 		return FString(TEXT("Editor/Matinee"));
 	}
@@ -81,22 +81,22 @@ public:
 	void OnCurveLabelClicked( UObject* CurveObject );
 
 	/** FGCObject interface */
-	void AddReferencedObjects( FReferenceCollector& Collector ) OVERRIDE;
+	void AddReferencedObjects( FReferenceCollector& Collector ) override;
 
 	/** 
 	 * Starts playing the current sequence. 
 	 * @param bPlayLoop		Whether or not we should play the looping section.
 	 * @param bPlayForward	true if we should play forwards, or false for reverse
 	 */
-	void StartPlaying(bool bPlayLoop, bool bPlayForward) OVERRIDE;
+	void StartPlaying(bool bPlayLoop, bool bPlayForward) override;
 
-	void ResumePlaying() OVERRIDE;
+	void ResumePlaying() override;
 
 	/** Stops playing the current sequence. */
-	void StopPlaying() OVERRIDE;
+	void StopPlaying() override;
 
 	/** Starts recording the current sequence */
-	void StartRecordingMovie() OVERRIDE;
+	void StartRecordingMovie() override;
 
 	// Menu handlers
 	void OnMenuAddKey();
@@ -429,7 +429,7 @@ public:
 	 * @param	TrackToSelect		The interp track to select.
 	 * @param	bDeselectPreviousTracks	If true, then all previously-selected tracks will be deselected. Defaults to true.
 	 */
-	void SelectTrack( UInterpGroup* Group, UInterpTrack* TrackToSelect, bool bDeselectPreviousTracks = true ) OVERRIDE;
+	void SelectTrack( UInterpGroup* Group, UInterpTrack* TrackToSelect, bool bDeselectPreviousTracks = true ) override;
 
 	/**
 	 * Selects the given group.
@@ -670,7 +670,7 @@ public:
 	/**
 	 * Dirty the contents of the track window viewports
 	 */
-	void InvalidateTrackWindowViewports() OVERRIDE;
+	void InvalidateTrackWindowViewports() override;
 
 	/**
 	 * Updates the contents of the property window based on which groups or tracks are selected if any. 
@@ -727,10 +727,10 @@ public:
 	void RepairHierarchyProblems();
 
 
-	bool KeyIsInSelection(class UInterpGroup* InGroup, UInterpTrack* InTrack, int32 InKeyIndex) OVERRIDE;
-	void AddKeyToSelection(class UInterpGroup* InGroup, UInterpTrack* InTrack, int32 InKeyIndex, bool bAutoWind) OVERRIDE;
-	void RemoveKeyFromSelection(class UInterpGroup* InGroup, UInterpTrack* InTrack, int32 InKeyIndex) OVERRIDE;
-	void ClearKeySelection() OVERRIDE;
+	bool KeyIsInSelection(class UInterpGroup* InGroup, UInterpTrack* InTrack, int32 InKeyIndex) override;
+	void AddKeyToSelection(class UInterpGroup* InGroup, UInterpTrack* InTrack, int32 InKeyIndex, bool bAutoWind) override;
+	void RemoveKeyFromSelection(class UInterpGroup* InGroup, UInterpTrack* InTrack, int32 InKeyIndex) override;
+	void ClearKeySelection() override;
 	void CalcSelectedKeyRange(float& OutStartTime, float& OutEndTime);
 	void SelectKeysInLoopSection();
 
@@ -758,7 +758,7 @@ public:
 	 *
 	 * There must be one and only one track selected for a keyframe to be added.
 	 */
-	void AddKey() OVERRIDE;
+	void AddKey() override;
 
 	/** 
 	 * Call utility to split an animation in the selected AnimControl track. 
@@ -963,15 +963,15 @@ public:
 	void SetInterpEnd(float NewInterpLength);
 	void MoveLoopMarker(float NewMarkerPos, bool bIsStart);
 
-	void BeginDrag3DHandle(UInterpGroup* Group, int32 TrackIndex) OVERRIDE;
-	void Move3DHandle(UInterpGroup* Group, int32 TrackIndex, int32 KeyIndex, bool bArriving, const FVector& Delta) OVERRIDE;
-	void EndDrag3DHandle() OVERRIDE;
-	void MoveInitialPosition(const FVector& Delta, const FRotator& DeltaRot) OVERRIDE;
+	void BeginDrag3DHandle(UInterpGroup* Group, int32 TrackIndex) override;
+	void Move3DHandle(UInterpGroup* Group, int32 TrackIndex, int32 KeyIndex, bool bArriving, const FVector& Delta) override;
+	void EndDrag3DHandle() override;
+	void MoveInitialPosition(const FVector& Delta, const FRotator& DeltaRot) override;
 
-	void ActorModified( bool bUpdateViewportTransform = true ) OVERRIDE;
-	void ActorSelectionChange( const bool bClearSelectionIfInvalid = true ) OVERRIDE;
-	void CamMoved(const FVector& NewCamLocation, const FRotator& NewCamRotation) OVERRIDE;
-	bool ProcessKeyPress(FKey Key, bool bCtrlDown, bool bAltDown) OVERRIDE;
+	void ActorModified( bool bUpdateViewportTransform = true ) override;
+	void ActorSelectionChange( const bool bClearSelectionIfInvalid = true ) override;
+	void CamMoved(const FVector& NewCamLocation, const FRotator& NewCamRotation) override;
+	bool ProcessKeyPress(FKey Key, bool bCtrlDown, bool bAltDown) override;
 
 	void InterpEdUndo();
 	void InterpEdRedo();
@@ -980,8 +980,8 @@ public:
 	void MoveActiveUp();
 	void MoveActiveDown();
 
-	void DrawTracks3D(const FSceneView* View, FPrimitiveDrawInterface* PDI) OVERRIDE;
-	void DrawModeHUD(FLevelEditorViewportClient* ViewportClient,FViewport* Viewport,const FSceneView* View,FCanvas* Canvas) OVERRIDE;
+	void DrawTracks3D(const FSceneView* View, FPrimitiveDrawInterface* PDI) override;
+	void DrawModeHUD(FLevelEditorViewportClient* ViewportClient,FViewport* Viewport,const FSceneView* View,FCanvas* Canvas) override;
 
 	void TickInterp(float DeltaSeconds);
 
@@ -1177,16 +1177,16 @@ public:
 	}
 
 	/** Toggles whether or not to display the menu*/
-	void ToggleRecordMenuDisplay (void) OVERRIDE
+	void ToggleRecordMenuDisplay (void) override
 	{
 		bDisplayRecordingMenu = !bDisplayRecordingMenu;
 	}
 
 	/**Preps Matinee to record/stop-recording realtime camera movement*/
-	void ToggleRecordInterpValues(void) OVERRIDE;
+	void ToggleRecordInterpValues(void) override;
 
 	/**If true, real time camera recording mode has been enabled*/
-	bool IsRecordingInterpValues (void) const OVERRIDE;
+	bool IsRecordingInterpValues (void) const override;
 
 	/**Helper function to properly shut down matinee recording*/
 	void StopRecordingInterpValues(void);
@@ -1206,16 +1206,16 @@ public:
 	/**
 	 * Resets the recording menu value to the default
 	 */
-	void ResetRecordingMenuValue(FEditorViewportClient* InClient) OVERRIDE;
+	void ResetRecordingMenuValue(FEditorViewportClient* InClient) override;
 
 	/**
 	 * Determines whether only the first click event is allowed or all repeat events are allowed
 	 * @return - true, if the value should change multiple times.  false, if the user should have to release and reclick
 	 */
-	bool IsRecordMenuChangeAllowedRepeat (void) const OVERRIDE;
+	bool IsRecordMenuChangeAllowedRepeat (void) const override;
 	
 	/** Sets the record mode for matinee */
-	void SetRecordMode(const uint32 InNewMode) OVERRIDE;
+	void SetRecordMode(const uint32 InNewMode) override;
 
 	/** Returns The time that sampling should start at */
 	const double GetRecordingStartTime (void) const;
@@ -1224,16 +1224,16 @@ public:
 	const double GetRecordingEndTime (void) const;
 
 	/**Return the number of samples we're keeping around for roll smoothing*/
-	int32 GetNumRecordRollSmoothingSamples (void) const OVERRIDE { return RecordRollSmoothingSamples; }
+	int32 GetNumRecordRollSmoothingSamples (void) const override { return RecordRollSmoothingSamples; }
 	/**Return the number of samples we're keeping around for roll smoothing*/
-	int32 GetNumRecordPitchSmoothingSamples (void) const OVERRIDE { return RecordPitchSmoothingSamples; }
+	int32 GetNumRecordPitchSmoothingSamples (void) const override { return RecordPitchSmoothingSamples; }
 	/**Returns the current movement scheme we're using for the camera*/
-	int32 GetCameraMovementScheme(void) const OVERRIDE { return RecordCameraMovementScheme; }
+	int32 GetCameraMovementScheme(void) const override { return RecordCameraMovementScheme; }
 
 	/** Save record settings for next run */
 	void SaveRecordingSettings(const FCameraControllerConfig& InCameraConfig);
 	/** Load record settings for next run */
-	void LoadRecordingSettings(FCameraControllerConfig& InCameraConfig) OVERRIDE;
+	void LoadRecordingSettings(FCameraControllerConfig& InCameraConfig) override;
 
 	/**
 	 * Access function to appropriate camera actor
@@ -1491,9 +1491,9 @@ private:
 	TSharedPtr<SWidget> CreateMarkerMenu(EMatineeMarkerType::Type MarkerType);
 
 	/** Accessors */
-	AMatineeActor* GetMatineeActor() OVERRIDE { return MatineeActor; }
-	UInterpData* GetInterpData() OVERRIDE { return IData; }
-	bool Hide3DTrackView() OVERRIDE { return bHide3DTrackView; }
+	AMatineeActor* GetMatineeActor() override { return MatineeActor; }
+	UInterpData* GetInterpData() override { return IData; }
+	bool Hide3DTrackView() override { return bHide3DTrackView; }
 
 	/** List of open tool panels; used to ensure only one exists at any one time */
 	TMap< FName, TWeakPtr<SDockableTab> > SpawnedToolPanels;
@@ -1547,7 +1547,7 @@ public:
 	 *	@param	TrackToAddKey	Interp track we are creating the key for
 	 *	@param	bCommitKeys		Setting to true will finalize the Key creation process
 	 */
-	void FinishAddKey(UInterpTrack *TrackToAddKey, bool bCommitKeys) OVERRIDE;
+	void FinishAddKey(UInterpTrack *TrackToAddKey, bool bCommitKeys) override;
 
 public:
 	/** Constants */
diff --git a/Engine/Source/Editor/Matinee/Private/MatineeFilterButton.h b/Engine/Source/Editor/Matinee/Private/MatineeFilterButton.h
--- a/Engine/Source/Editor/Matinee/Private/MatineeFilterButton.h
+++ b/Engine/Source/Editor/Matinee/Private/MatineeFilterButton.h
@@ -35,7 +35,7 @@ public:
 	void Construct(const FArguments& InArgs);
 
 	// SWidget interface
-	virtual FReply OnMouseButtonUp( const FGeometry& MyGeometry, const FPointerEvent& MouseEvent ) OVERRIDE;
+	virtual FReply OnMouseButtonUp( const FGeometry& MyGeometry, const FPointerEvent& MouseEvent ) override;
 
 private:
 	/** Called when the checked state has changed */
diff --git a/Engine/Source/Editor/Matinee/Private/MatineeModule.cpp b/Engine/Source/Editor/Matinee/Private/MatineeModule.cpp
--- a/Engine/Source/Editor/Matinee/Private/MatineeModule.cpp
+++ b/Engine/Source/Editor/Matinee/Private/MatineeModule.cpp
@@ -22,26 +22,26 @@ public:
 	}
 
 	/** Called right after the module DLL has been loaded and the module object has been created */
-	virtual void StartupModule() OVERRIDE
+	virtual void StartupModule() override
 	{
 		MenuExtensibilityManager = MakeShareable(new FExtensibilityManager);
 		ToolBarExtensibilityManager = MakeShareable(new FExtensibilityManager);
 	}
 
 	/** Called before the module is unloaded, right before the module object is destroyed. */
-	virtual void ShutdownModule() OVERRIDE
+	virtual void ShutdownModule() override
 	{
 		MenuExtensibilityManager.Reset();
 		ToolBarExtensibilityManager.Reset();
 	}
 
 	DECLARE_DERIVED_EVENT(FMatineeModule, IMatineeModule::FMatineeEditorOpenedEvent, FMatineeEditorOpenedEvent);
-	virtual FMatineeEditorOpenedEvent& OnMatineeEditorOpened() OVERRIDE
+	virtual FMatineeEditorOpenedEvent& OnMatineeEditorOpened() override
 	{
 		return MatineeEditorOpenedEvent;
 	}
 
-	virtual TSharedRef<IMatinee> CreateMatinee(const EToolkitMode::Type Mode, const TSharedPtr< IToolkitHost >& InitToolkitHost, AMatineeActor* MatineeActor) OVERRIDE
+	virtual TSharedRef<IMatinee> CreateMatinee(const EToolkitMode::Type Mode, const TSharedPtr< IToolkitHost >& InitToolkitHost, AMatineeActor* MatineeActor) override
 	{
 		TSharedRef<FMatinee> MatineeEd(new FMatinee());
 		MatineeEd->InitMatinee(Mode, InitToolkitHost, MatineeActor);
diff --git a/Engine/Source/Editor/Matinee/Public/IMatinee.h b/Engine/Source/Editor/Matinee/Public/IMatinee.h
--- a/Engine/Source/Editor/Matinee/Public/IMatinee.h
+++ b/Engine/Source/Editor/Matinee/Public/IMatinee.h
@@ -88,7 +88,7 @@ class IMatinee : public IMatineeBase, public FAssetEditorToolkit
 {
 public:
 
-	bool Show(bool bShow) OVERRIDE {return false; }
-	void Close(bool bForce) OVERRIDE { CloseWindow(); }
+	bool Show(bool bShow) override {return false; }
+	void Close(bool bForce) override { CloseWindow(); }
 };
 	
diff --git a/Engine/Source/Editor/MeshPaint/Private/MeshPaintEdMode.cpp b/Engine/Source/Editor/MeshPaint/Private/MeshPaintEdMode.cpp
--- a/Engine/Source/Editor/MeshPaint/Private/MeshPaintEdMode.cpp
+++ b/Engine/Source/Editor/MeshPaint/Private/MeshPaintEdMode.cpp
@@ -44,7 +44,7 @@ class FMeshPaintBatchedElementParameters : public FBatchedElementParameters
 public:
 
 	/** Binds vertex and pixel shaders for this element */
-	virtual void BindShaders_RenderThread(FRHICommandList& RHICmdList, const FMatrix& InTransform, const float InGamma, const FMatrix& ColorWeights, const FTexture* Texture ) OVERRIDE
+	virtual void BindShaders_RenderThread(FRHICommandList& RHICmdList, const FMatrix& InTransform, const float InGamma, const FMatrix& ColorWeights, const FTexture* Texture ) override
 	{
 		MeshPaintRendering::SetMeshPaintShaders_RenderThread(RHICmdList, InTransform, InGamma, ShaderParams );
 	}
@@ -63,7 +63,7 @@ class FMeshPaintDilateBatchedElementParameters : public FBatchedElementParameter
 public:
 
 	/** Binds vertex and pixel shaders for this element */
-	virtual void BindShaders_RenderThread(FRHICommandList& RHICmdList, const FMatrix& InTransform, const float InGamma, const FMatrix& ColorWeights, const FTexture* Texture ) OVERRIDE
+	virtual void BindShaders_RenderThread(FRHICommandList& RHICmdList, const FMatrix& InTransform, const float InGamma, const FMatrix& ColorWeights, const FTexture* Texture ) override
 	{
 		MeshPaintRendering::SetMeshPaintDilateShaders_RenderThread(RHICmdList, InTransform, InGamma, ShaderParams );
 	}
diff --git a/Engine/Source/Editor/MeshPaint/Private/SMeshPaint.h b/Engine/Source/Editor/MeshPaint/Private/SMeshPaint.h
--- a/Engine/Source/Editor/MeshPaint/Private/SMeshPaint.h
+++ b/Engine/Source/Editor/MeshPaint/Private/SMeshPaint.h
@@ -39,17 +39,17 @@ namespace EMeshPaintWriteColorChannels
 class FMeshPaintToolKit : public FModeToolkit
 {
 public:
-	virtual void RegisterTabSpawners(const TSharedRef<class FTabManager>& TabManager) OVERRIDE;
-	virtual void UnregisterTabSpawners(const TSharedRef<class FTabManager>& TabManager) OVERRIDE;
+	virtual void RegisterTabSpawners(const TSharedRef<class FTabManager>& TabManager) override;
+	virtual void UnregisterTabSpawners(const TSharedRef<class FTabManager>& TabManager) override;
 
 	/** Initializes the geometry mode toolkit */
 	virtual void Init(const TSharedPtr< class IToolkitHost >& InitToolkitHost);
 
 	/** IToolkit interface */
-	virtual FName GetToolkitFName() const OVERRIDE;
-	virtual FText GetBaseToolkitName() const OVERRIDE;
-	virtual class FEdMode* GetEditorMode() const OVERRIDE;
-	virtual TSharedPtr<class SWidget> GetInlineContent() const OVERRIDE;
+	virtual FName GetToolkitFName() const override;
+	virtual FText GetBaseToolkitName() const override;
+	virtual class FEdMode* GetEditorMode() const override;
+	virtual TSharedPtr<class SWidget> GetInlineContent() const override;
 
 private:
 	/** Geometry tools widget */
@@ -68,7 +68,7 @@ public:
 	/** SCompoundWidget functions */
 	void Construct(const FArguments& InArgs, TSharedRef<FMeshPaintToolKit> InParentToolkit);
 
-	virtual void Tick( const FGeometry& AllottedGeometry, const double InCurrentTime, const float InDeltaTime ) OVERRIDE;
+	virtual void Tick( const FGeometry& AllottedGeometry, const double InCurrentTime, const float InDeltaTime ) override;
 
 protected:
 
diff --git a/Engine/Source/Editor/MeshPaint/Public/MeshPaintEdMode.h b/Engine/Source/Editor/MeshPaint/Public/MeshPaintEdMode.h
--- a/Engine/Source/Editor/MeshPaint/Public/MeshPaintEdMode.h
+++ b/Engine/Source/Editor/MeshPaint/Public/MeshPaintEdMode.h
@@ -353,27 +353,27 @@ public:
 	virtual ~FEdModeMeshPaint();
 
 	/** FGCObject interface */
-	virtual void AddReferencedObjects( FReferenceCollector& Collector ) OVERRIDE;
+	virtual void AddReferencedObjects( FReferenceCollector& Collector ) override;
 
 	// FEdMode interface
-	virtual bool UsesToolkits() const OVERRIDE;
-	virtual void Enter() OVERRIDE;
-	virtual void Exit() OVERRIDE;
-	virtual bool MouseMove( FLevelEditorViewportClient* ViewportClient, FViewport* Viewport, int32 x, int32 y ) OVERRIDE;
-	virtual bool CapturedMouseMove( FLevelEditorViewportClient* InViewportClient, FViewport* InViewport, int32 InMouseX, int32 InMouseY ) OVERRIDE;
-	virtual bool StartTracking(FLevelEditorViewportClient* InViewportClient, FViewport* InViewport) OVERRIDE;
-	virtual bool EndTracking(FLevelEditorViewportClient* InViewportClient, FViewport* InViewport) OVERRIDE;
-	virtual bool InputKey( FLevelEditorViewportClient* InViewportClient, FViewport* InViewport, FKey InKey, EInputEvent InEvent ) OVERRIDE;
-	virtual bool InputDelta( FLevelEditorViewportClient* InViewportClient, FViewport* InViewport, FVector& InDrag, FRotator& InRot, FVector& InScale ) OVERRIDE;
-	virtual void PostUndo() OVERRIDE;
-	virtual void Render( const FSceneView* View, FViewport* Viewport, FPrimitiveDrawInterface* PDI ) OVERRIDE;
-	virtual bool Select( AActor* InActor, bool bInSelected ) OVERRIDE;
-	virtual bool IsSelectionAllowed( AActor* InActor, bool bInSelection ) const OVERRIDE;
-	virtual void ActorSelectionChangeNotify() OVERRIDE;
-	virtual bool AllowWidgetMove() OVERRIDE { return false; }
-	virtual bool ShouldDrawWidget() const OVERRIDE { return false; }
-	virtual bool UsesTransformWidget() const OVERRIDE { return false; }
-	virtual void Tick( FLevelEditorViewportClient* ViewportClient, float DeltaTime ) OVERRIDE;
+	virtual bool UsesToolkits() const override;
+	virtual void Enter() override;
+	virtual void Exit() override;
+	virtual bool MouseMove( FLevelEditorViewportClient* ViewportClient, FViewport* Viewport, int32 x, int32 y ) override;
+	virtual bool CapturedMouseMove( FLevelEditorViewportClient* InViewportClient, FViewport* InViewport, int32 InMouseX, int32 InMouseY ) override;
+	virtual bool StartTracking(FLevelEditorViewportClient* InViewportClient, FViewport* InViewport) override;
+	virtual bool EndTracking(FLevelEditorViewportClient* InViewportClient, FViewport* InViewport) override;
+	virtual bool InputKey( FLevelEditorViewportClient* InViewportClient, FViewport* InViewport, FKey InKey, EInputEvent InEvent ) override;
+	virtual bool InputDelta( FLevelEditorViewportClient* InViewportClient, FViewport* InViewport, FVector& InDrag, FRotator& InRot, FVector& InScale ) override;
+	virtual void PostUndo() override;
+	virtual void Render( const FSceneView* View, FViewport* Viewport, FPrimitiveDrawInterface* PDI ) override;
+	virtual bool Select( AActor* InActor, bool bInSelected ) override;
+	virtual bool IsSelectionAllowed( AActor* InActor, bool bInSelection ) const override;
+	virtual void ActorSelectionChangeNotify() override;
+	virtual bool AllowWidgetMove() override { return false; }
+	virtual bool ShouldDrawWidget() const override { return false; }
+	virtual bool UsesTransformWidget() const override { return false; }
+	virtual void Tick( FLevelEditorViewportClient* ViewportClient, float DeltaTime ) override;
 	// End of FEdMode interface
 
 
diff --git a/Engine/Source/Editor/NewLevelDialog/Private/NewLevelDialogModule.cpp b/Engine/Source/Editor/NewLevelDialog/Private/NewLevelDialogModule.cpp
--- a/Engine/Source/Editor/NewLevelDialog/Private/NewLevelDialogModule.cpp
+++ b/Engine/Source/Editor/NewLevelDialog/Private/NewLevelDialogModule.cpp
@@ -46,7 +46,7 @@ public:
 
 	virtual FIntPoint GetSize() const { return Size; }
 	virtual class FSlateShaderResource* GetViewportRenderTargetTexture() const { return ShaderResource ? (FSlateShaderResource*)this : NULL; }
-	virtual bool RequiresVsync() const OVERRIDE { return false; }
+	virtual bool RequiresVsync() const override { return false; }
 
 	// FSlateShaderResource implementation.
 	virtual uint32 GetWidth() const { return Size.X; }
@@ -259,7 +259,7 @@ private:
 		}
 	}
 
-	virtual FReply OnKeyDown( const FGeometry& MyGeometry, const FKeyboardEvent& InKeyboardEvent ) OVERRIDE
+	virtual FReply OnKeyDown( const FGeometry& MyGeometry, const FKeyboardEvent& InKeyboardEvent ) override
 	{
 		if( InKeyboardEvent.GetKey() == EKeys::Escape )
 		{
diff --git a/Engine/Source/Editor/OnlineBlueprintSupport/Classes/K2Node_LatentOnlineCall.h b/Engine/Source/Editor/OnlineBlueprintSupport/Classes/K2Node_LatentOnlineCall.h
--- a/Engine/Source/Editor/OnlineBlueprintSupport/Classes/K2Node_LatentOnlineCall.h
+++ b/Engine/Source/Editor/OnlineBlueprintSupport/Classes/K2Node_LatentOnlineCall.h
@@ -11,6 +11,6 @@ class ONLINEBLUEPRINTSUPPORT_API UK2Node_LatentOnlineCall : public UK2Node_BaseA
 	GENERATED_UCLASS_BODY()
 
 	// UEdGraphNode interface
-	virtual void GetMenuEntries(FGraphContextMenuBuilder& ContextMenuBuilder) const OVERRIDE;
+	virtual void GetMenuEntries(FGraphContextMenuBuilder& ContextMenuBuilder) const override;
 	// End of UEdGraphNode interface
 };
diff --git a/Engine/Source/Editor/OnlineBlueprintSupport/Classes/K2Node_LeaderboardFlush.h b/Engine/Source/Editor/OnlineBlueprintSupport/Classes/K2Node_LeaderboardFlush.h
--- a/Engine/Source/Editor/OnlineBlueprintSupport/Classes/K2Node_LeaderboardFlush.h
+++ b/Engine/Source/Editor/OnlineBlueprintSupport/Classes/K2Node_LeaderboardFlush.h
@@ -10,11 +10,11 @@ class ONLINEBLUEPRINTSUPPORT_API UK2Node_LeaderboardFlush : public UK2Node_BaseA
 	GENERATED_UCLASS_BODY()
 
 	// UEdGraphNode interface
-	virtual FString GetTooltip() const OVERRIDE;
-	virtual FText GetNodeTitle(ENodeTitleType::Type TitleType) const OVERRIDE;
+	virtual FString GetTooltip() const override;
+	virtual FText GetNodeTitle(ENodeTitleType::Type TitleType) const override;
 	// End of UEdGraphNode interface
 
 	// UK2Node_BaseAsyncTask interface
-	virtual FString GetCategoryName() OVERRIDE;
+	virtual FString GetCategoryName() override;
 	// End of UK2Node_BaseAsyncTask interface
 };
diff --git a/Engine/Source/Editor/PListEditor/Private/PListNodeArray.h b/Engine/Source/Editor/PListEditor/Private/PListNodeArray.h
--- a/Engine/Source/Editor/PListEditor/Private/PListNodeArray.h
+++ b/Engine/Source/Editor/PListEditor/Private/PListNodeArray.h
@@ -18,49 +18,49 @@ public:
 public:
 
 	/** Validation check */
-	virtual bool IsValid() OVERRIDE;
+	virtual bool IsValid() override;
 
 	/** Returns the array of children */
-	virtual TArray<TSharedPtr<IPListNode> >& GetChildren() OVERRIDE;
+	virtual TArray<TSharedPtr<IPListNode> >& GetChildren() override;
 
 	/** Adds a child to the internal array of the class */
-	virtual void AddChild(TSharedPtr<IPListNode> InChild) OVERRIDE;
+	virtual void AddChild(TSharedPtr<IPListNode> InChild) override;
 
 	/** Gets the type of the node via enum */
-	virtual EPLNTypes GetType() OVERRIDE;
+	virtual EPLNTypes GetType() override;
 
 	/** Determines whether the node needs to generate widgets for columns, or just use the whole row without columns */
-	virtual bool UsesColumns() OVERRIDE;
+	virtual bool UsesColumns() override;
 
 	/** Generates a widget for a TableView row */
-	virtual TSharedRef<ITableRow> GenerateWidget(const TSharedRef<STableViewBase>& OwnerTable) OVERRIDE;
+	virtual TSharedRef<ITableRow> GenerateWidget(const TSharedRef<STableViewBase>& OwnerTable) override;
 
 	/** Generate a widget for the specified column name */
-	virtual TSharedRef<SWidget> GenerateWidgetForColumn(const FName& ColumnName, int32 Depth, ITableRow* RowPtr) OVERRIDE;
+	virtual TSharedRef<SWidget> GenerateWidgetForColumn(const FName& ColumnName, int32 Depth, ITableRow* RowPtr) override;
 
 	/** Gets an XML representation of the node's internals */
-	virtual FString ToXML(const int32 indent = 0, bool bOutputKey = true) OVERRIDE;
+	virtual FString ToXML(const int32 indent = 0, bool bOutputKey = true) override;
 
 	/** Refreshes anything necessary in the node, such as a text used in information displaying */
-	virtual void Refresh() OVERRIDE;
+	virtual void Refresh() override;
 
 	/** Calculate how many key/value pairs exist for this node. */
-	virtual int32 GetNumPairs() OVERRIDE;
+	virtual int32 GetNumPairs() override;
 
 	/** Called when the filter changes */
-	virtual void OnFilterTextChanged(FString NewFilter) OVERRIDE;
+	virtual void OnFilterTextChanged(FString NewFilter) override;
 
 	/** When parents are refreshed, they can set the index of their children for proper displaying */
-	virtual void SetIndex(int32 NewIndex) OVERRIDE;
+	virtual void SetIndex(int32 NewIndex) override;
 
 public:
 	
 	/** Gets the overlay brush to use */
-	const FSlateBrush* GetOverlayBrush() OVERRIDE;
+	const FSlateBrush* GetOverlayBrush() override;
 	/** Sets the KeyString of the node, needed for telling children to set their keys */
-	virtual void SetKey(FString NewString) OVERRIDE;
+	virtual void SetKey(FString NewString) override;
 	/** Sets a flag denoting whether the element is in an array. Default do nothing */
-	virtual void SetArrayMember(bool bArrayMember) OVERRIDE;
+	virtual void SetArrayMember(bool bArrayMember) override;
 
 private:
 	/** Delegate: When key string changes */
diff --git a/Engine/Source/Editor/PListEditor/Private/PListNodeBoolean.h b/Engine/Source/Editor/PListEditor/Private/PListNodeBoolean.h
--- a/Engine/Source/Editor/PListEditor/Private/PListNodeBoolean.h
+++ b/Engine/Source/Editor/PListEditor/Private/PListNodeBoolean.h
@@ -18,51 +18,51 @@ public:
 public:
 
 	/** Validation check */
-	virtual bool IsValid() OVERRIDE;
+	virtual bool IsValid() override;
 
 	/** Returns the array of children */
-	virtual TArray<TSharedPtr<IPListNode> >& GetChildren() OVERRIDE;
+	virtual TArray<TSharedPtr<IPListNode> >& GetChildren() override;
 
 	/** Adds a child to the internal array of the class */
-	virtual void AddChild(TSharedPtr<IPListNode> InChild) OVERRIDE;
+	virtual void AddChild(TSharedPtr<IPListNode> InChild) override;
 
 	/** Gets the type of the node via enum */
-	virtual EPLNTypes GetType() OVERRIDE;
+	virtual EPLNTypes GetType() override;
 
 	/** Determines whether the node needs to generate widgets for columns, or just use the whole row without columns */
-	virtual bool UsesColumns() OVERRIDE;
+	virtual bool UsesColumns() override;
 
 	/** Generates a widget for a TableView row */
-	virtual TSharedRef<ITableRow> GenerateWidget(const TSharedRef<STableViewBase>& OwnerTable) OVERRIDE;
+	virtual TSharedRef<ITableRow> GenerateWidget(const TSharedRef<STableViewBase>& OwnerTable) override;
 
 	/** Generate a widget for the specified column name */
-	virtual TSharedRef<SWidget> GenerateWidgetForColumn(const FName& ColumnName, int32 Depth, ITableRow* RowPtr) OVERRIDE;
+	virtual TSharedRef<SWidget> GenerateWidgetForColumn(const FName& ColumnName, int32 Depth, ITableRow* RowPtr) override;
 
 	/** Gets an XML representation of the node's internals */
-	virtual FString ToXML(const int32 indent = 0, bool bOutputKey = true) OVERRIDE;
+	virtual FString ToXML(const int32 indent = 0, bool bOutputKey = true) override;
 
 	/** Refreshes anything necessary in the node, such as a text used in information displaying */
-	virtual void Refresh() OVERRIDE;
+	virtual void Refresh() override;
 
 	/** Calculate how many key/value pairs exist for this node. */
-	virtual int32 GetNumPairs() OVERRIDE;
+	virtual int32 GetNumPairs() override;
 
 	/** Called when the filter changes */
-	virtual void OnFilterTextChanged(FString NewFilter) OVERRIDE;
+	virtual void OnFilterTextChanged(FString NewFilter) override;
 
 	/** When parents are refreshed, they can set the index of their children for proper displaying */
-	virtual void SetIndex(int32 NewIndex) OVERRIDE;
+	virtual void SetIndex(int32 NewIndex) override;
 
 public:
 	
 	/** Gets the brush to use based on filter status */
-	const FSlateBrush* GetOverlayBrush() OVERRIDE;
+	const FSlateBrush* GetOverlayBrush() override;
 	/** Sets the KeyString of the node, needed for telling children to set their keys */
-	virtual void SetKey(FString NewString) OVERRIDE;
+	virtual void SetKey(FString NewString) override;
 	/** Sets the value of the boolean */
-	virtual void SetValue(bool bNewValue) OVERRIDE;
+	virtual void SetValue(bool bNewValue) override;
 	/** Sets a flag denoting whether the element is in an array */
-	virtual void SetArrayMember(bool bArrayMember) OVERRIDE;
+	virtual void SetArrayMember(bool bArrayMember) override;
 
 private:
 
diff --git a/Engine/Source/Editor/PListEditor/Private/PListNodeDictionary.h b/Engine/Source/Editor/PListEditor/Private/PListNodeDictionary.h
--- a/Engine/Source/Editor/PListEditor/Private/PListNodeDictionary.h
+++ b/Engine/Source/Editor/PListEditor/Private/PListNodeDictionary.h
@@ -18,47 +18,47 @@ public:
 public:
 
 	/** Validation check */
-	virtual bool IsValid() OVERRIDE;
+	virtual bool IsValid() override;
 
 	/** Returns the array of children */
-	virtual TArray<TSharedPtr<IPListNode> >& GetChildren() OVERRIDE;
+	virtual TArray<TSharedPtr<IPListNode> >& GetChildren() override;
 
 	/** Adds a child to the internal array of the class */
-	virtual void AddChild(TSharedPtr<IPListNode> InChild) OVERRIDE;
+	virtual void AddChild(TSharedPtr<IPListNode> InChild) override;
 
 	/** Gets the type of the node via enum */
-	virtual EPLNTypes GetType() OVERRIDE;
+	virtual EPLNTypes GetType() override;
 
 	/** Determines whether the node needs to generate widgets for columns, or just use the whole row without columns */
-	virtual bool UsesColumns() OVERRIDE;
+	virtual bool UsesColumns() override;
 
 	/** Generates a widget for a TableView row */
-	virtual TSharedRef<ITableRow> GenerateWidget(const TSharedRef<STableViewBase>& OwnerTable) OVERRIDE;
+	virtual TSharedRef<ITableRow> GenerateWidget(const TSharedRef<STableViewBase>& OwnerTable) override;
 
 	/** Generate a widget for the specified column name */
-	virtual TSharedRef<SWidget> GenerateWidgetForColumn(const FName& ColumnName, int32 Depth, ITableRow* RowPtr) OVERRIDE;
+	virtual TSharedRef<SWidget> GenerateWidgetForColumn(const FName& ColumnName, int32 Depth, ITableRow* RowPtr) override;
 
 	/** Gets an XML representation of the node's internals */
-	virtual FString ToXML(const int32 indent = 0, bool bOutputKey = true) OVERRIDE;
+	virtual FString ToXML(const int32 indent = 0, bool bOutputKey = true) override;
 
 	/** Refreshes anything necessary in the node, such as a text used in information displaying */
-	virtual void Refresh() OVERRIDE;
+	virtual void Refresh() override;
 
 	/** Calculate how many key/value pairs exist for this node. */
-	virtual int32 GetNumPairs() OVERRIDE;
+	virtual int32 GetNumPairs() override;
 
 	/** Called when the filter changes */
-	virtual void OnFilterTextChanged(FString NewFilter) OVERRIDE;
+	virtual void OnFilterTextChanged(FString NewFilter) override;
 
 	/** When parents are refreshed, they can set the index of their children for proper displaying */
-	virtual void SetIndex(int32 NewIndex) OVERRIDE;
+	virtual void SetIndex(int32 NewIndex) override;
 
 public:
 
 	/** Sets a flag denoting whether the element is in an array. Default do nothing */
-	virtual void SetArrayMember(bool bArrayMember) OVERRIDE;
+	virtual void SetArrayMember(bool bArrayMember) override;
 	/** Gets the overlay brush dynamically */
-	virtual const FSlateBrush* GetOverlayBrush() OVERRIDE;
+	virtual const FSlateBrush* GetOverlayBrush() override;
 
 private:
 
diff --git a/Engine/Source/Editor/PListEditor/Private/PListNodeFile.h b/Engine/Source/Editor/PListEditor/Private/PListNodeFile.h
--- a/Engine/Source/Editor/PListEditor/Private/PListNodeFile.h
+++ b/Engine/Source/Editor/PListEditor/Private/PListNodeFile.h
@@ -18,40 +18,40 @@ public:
 public:
 
 	/** Validation check */
-	virtual bool IsValid() OVERRIDE;
+	virtual bool IsValid() override;
 
 	/** Returns the array of children */
-	virtual TArray<TSharedPtr<IPListNode> >& GetChildren() OVERRIDE;
+	virtual TArray<TSharedPtr<IPListNode> >& GetChildren() override;
 
 	/** Adds a child to the internal array of the class */
-	virtual void AddChild(TSharedPtr<IPListNode> InChild) OVERRIDE;
+	virtual void AddChild(TSharedPtr<IPListNode> InChild) override;
 
 	/** Gets the type of the node via enum */
-	virtual EPLNTypes GetType() OVERRIDE;
+	virtual EPLNTypes GetType() override;
 
 	/** Determines whether the node needs to generate widgets for columns, or just use the whole row without columns */
-	virtual bool UsesColumns() OVERRIDE;
+	virtual bool UsesColumns() override;
 
 	/** Generates a widget for a TableView row */
-	virtual TSharedRef<ITableRow> GenerateWidget(const TSharedRef<STableViewBase>& OwnerTable) OVERRIDE;
+	virtual TSharedRef<ITableRow> GenerateWidget(const TSharedRef<STableViewBase>& OwnerTable) override;
 
 	/** Generate a widget for the specified column name */
-	virtual TSharedRef<SWidget> GenerateWidgetForColumn(const FName& ColumnName, int32 Depth, ITableRow* RowPtr) OVERRIDE;
+	virtual TSharedRef<SWidget> GenerateWidgetForColumn(const FName& ColumnName, int32 Depth, ITableRow* RowPtr) override;
 
 	/** Gets an XML representation of the node's internals */
-	virtual FString ToXML(const int32 indent = 0, bool bOutputKey = true) OVERRIDE;
+	virtual FString ToXML(const int32 indent = 0, bool bOutputKey = true) override;
 
 	/** Refreshes anything necessary in the node, such as a text used in information displaying */
-	virtual void Refresh() OVERRIDE;
+	virtual void Refresh() override;
 
 	/** Calculate how many key/value pairs exist for this node. */
-	virtual int32 GetNumPairs() OVERRIDE;
+	virtual int32 GetNumPairs() override;
 
 	/** Called when the filter changes */
-	virtual void OnFilterTextChanged(FString NewFilter) OVERRIDE;
+	virtual void OnFilterTextChanged(FString NewFilter) override;
 
 	/** When parents are refreshed, they can set the index of their children for proper displaying */
-	virtual void SetIndex(int32 NewIndex) OVERRIDE;
+	virtual void SetIndex(int32 NewIndex) override;
 
 private:
 
diff --git a/Engine/Source/Editor/PListEditor/Private/PListNodeString.h b/Engine/Source/Editor/PListEditor/Private/PListNodeString.h
--- a/Engine/Source/Editor/PListEditor/Private/PListNodeString.h
+++ b/Engine/Source/Editor/PListEditor/Private/PListNodeString.h
@@ -18,51 +18,51 @@ public:
 public:
 
 	/** Validation check */
-	virtual bool IsValid() OVERRIDE;
+	virtual bool IsValid() override;
 
 	/** Returns the array of children */
-	virtual TArray<TSharedPtr<IPListNode> >& GetChildren() OVERRIDE;
+	virtual TArray<TSharedPtr<IPListNode> >& GetChildren() override;
 
 	/** Adds a child to the internal array of the class */
-	virtual void AddChild(TSharedPtr<IPListNode> InChild) OVERRIDE;
+	virtual void AddChild(TSharedPtr<IPListNode> InChild) override;
 
 	/** Gets the type of the node via enum */
-	virtual EPLNTypes GetType() OVERRIDE;
+	virtual EPLNTypes GetType() override;
 
 	/** Determines whether the node needs to generate widgets for columns, or just use the whole row without columns */
-	virtual bool UsesColumns() OVERRIDE;
+	virtual bool UsesColumns() override;
 
 	/** Generates a widget for a TableView row */
-	virtual TSharedRef<ITableRow> GenerateWidget(const TSharedRef<STableViewBase>& OwnerTable) OVERRIDE;
+	virtual TSharedRef<ITableRow> GenerateWidget(const TSharedRef<STableViewBase>& OwnerTable) override;
 
 	/** Generate a widget for the specified column name */
-	virtual TSharedRef<SWidget> GenerateWidgetForColumn(const FName& ColumnName, int32 Depth, ITableRow* RowPtr) OVERRIDE;
+	virtual TSharedRef<SWidget> GenerateWidgetForColumn(const FName& ColumnName, int32 Depth, ITableRow* RowPtr) override;
 
 	/** Gets an XML representation of the node's internals */
-	virtual FString ToXML(const int32 indent = 0, bool bOutputKey = true) OVERRIDE;
+	virtual FString ToXML(const int32 indent = 0, bool bOutputKey = true) override;
 
 	/** Refreshes anything necessary in the node, such as a text used in information displaying */
-	virtual void Refresh() OVERRIDE;
+	virtual void Refresh() override;
 
 	/** Calculate how many key/value pairs exist for this node. */
-	virtual int32 GetNumPairs() OVERRIDE;
+	virtual int32 GetNumPairs() override;
 
 	/** Called when the filter changes */
-	virtual void OnFilterTextChanged(FString NewFilter) OVERRIDE;
+	virtual void OnFilterTextChanged(FString NewFilter) override;
 
 	/** When parents are refreshed, they can set the index of their children for proper displaying */
-	virtual void SetIndex(int32 NewIndex) OVERRIDE;
+	virtual void SetIndex(int32 NewIndex) override;
 
 public:
 	
 	/** Gets the brush to use based on filter status */
-	const FSlateBrush* GetOverlayBrush() OVERRIDE;
+	const FSlateBrush* GetOverlayBrush() override;
 	/** Sets the KeyString of the node, needed for telling children to set their keys */
-	virtual void SetKey(FString NewString) OVERRIDE;
+	virtual void SetKey(FString NewString) override;
 	/** Sets the value of the boolean */
-	virtual void SetValue(FString NewValue) OVERRIDE;
+	virtual void SetValue(FString NewValue) override;
 	/** Sets a flag denoting whether the element is in an array */
-	virtual void SetArrayMember(bool bArrayMember) OVERRIDE;
+	virtual void SetArrayMember(bool bArrayMember) override;
 
 private:
 
diff --git a/Engine/Source/Editor/PListEditor/Public/SPlistEditor.h b/Engine/Source/Editor/PListEditor/Public/SPlistEditor.h
--- a/Engine/Source/Editor/PListEditor/Public/SPlistEditor.h
+++ b/Engine/Source/Editor/PListEditor/Public/SPlistEditor.h
@@ -87,7 +87,7 @@ public:
 	TSharedPtr<FUICommandInfo> AddBooleanCommand;
 
 	/** Initialize Commands */
-	virtual void RegisterCommands() OVERRIDE
+	virtual void RegisterCommands() override
 	{
 		UI_COMMAND( NewCommand, "New", "Creates a new plist file", EUserInterfaceActionType::Button, FInputGesture(EModifierKey::Control, EKeys::N) );
 		UI_COMMAND( OpenCommand, "Open", "Opens an existing plist file", EUserInterfaceActionType::Button, FInputGesture(EModifierKey::Control, EKeys::O) );
@@ -117,9 +117,9 @@ public:
 	}
 	
 	// SWidget interface
-	virtual void Tick( const FGeometry& AllottedGeometry, const double InCurrentTime, const float InDeltaTime ) OVERRIDE;
-	virtual FReply OnKeyDown( const FGeometry& MyGeometry, const FKeyboardEvent& InKeyboardEvent ) OVERRIDE;
-	virtual FReply OnMouseButtonDown( const FGeometry& MyGeometry, const FPointerEvent& MouseEvent ) OVERRIDE;
+	virtual void Tick( const FGeometry& AllottedGeometry, const double InCurrentTime, const float InDeltaTime ) override;
+	virtual FReply OnKeyDown( const FGeometry& MyGeometry, const FKeyboardEvent& InKeyboardEvent ) override;
+	virtual FReply OnMouseButtonDown( const FGeometry& MyGeometry, const FPointerEvent& MouseEvent ) override;
 	// End of SWidget interface
 
 	/** Construct the widget when opened */
diff --git a/Engine/Source/Editor/PackagesDialog/Private/SPackagesDialog.h b/Engine/Source/Editor/PackagesDialog/Private/SPackagesDialog.h
--- a/Engine/Source/Editor/PackagesDialog/Private/SPackagesDialog.h
+++ b/Engine/Source/Editor/PackagesDialog/Private/SPackagesDialog.h
@@ -357,7 +357,7 @@ public:
 private:
 
 	// Begin SWidget Interface
-	virtual FReply OnKeyDown( const FGeometry& MyGeometry, const FKeyboardEvent& InKeyboardEvent ) OVERRIDE;
+	virtual FReply OnKeyDown( const FGeometry& MyGeometry, const FKeyboardEvent& InKeyboardEvent ) override;
 	// End SWidget Interface
 
 	/**
@@ -496,7 +496,7 @@ public:
 	void Construct( const FArguments& InArgs, const TSharedRef<STableViewBase>& InOwnerTableView );
 
 	/** Overridden from SMultiColumnTableRow.  Generates a widget for this column of the list row. */
-	virtual TSharedRef<SWidget> GenerateWidgetForColumn( const FName& ColumnName ) OVERRIDE;
+	virtual TSharedRef<SWidget> GenerateWidgetForColumn( const FName& ColumnName ) override;
 
 
 private:
diff --git a/Engine/Source/Editor/Persona/Private/AnimGraphNodeDetails.h b/Engine/Source/Editor/Persona/Private/AnimGraphNodeDetails.h
--- a/Engine/Source/Editor/Persona/Private/AnimGraphNodeDetails.h
+++ b/Engine/Source/Editor/Persona/Private/AnimGraphNodeDetails.h
@@ -14,7 +14,7 @@ public:
 	static TSharedRef<IDetailCustomization> MakeInstance();
 
 	// IDetailCustomization interface
-	virtual void CustomizeDetails(class IDetailLayoutBuilder& DetailBuilder) OVERRIDE;
+	virtual void CustomizeDetails(class IDetailLayoutBuilder& DetailBuilder) override;
 	// End of IDetailCustomization interface
 
 protected:
@@ -42,8 +42,8 @@ public:
 	static TSharedRef<IPropertyTypeCustomization> MakeInstance();
 
 	// IPropertyTypeCustomization interface
-	virtual void CustomizeHeader(TSharedRef<IPropertyHandle> StructPropertyHandle, class FDetailWidgetRow& HeaderRow, IPropertyTypeCustomizationUtils& StructCustomizationUtils) OVERRIDE;
-	virtual void CustomizeChildren(TSharedRef<IPropertyHandle> StructPropertyHandle, class IDetailChildrenBuilder& ChildBuilder, IPropertyTypeCustomizationUtils& StructCustomizationUtils) OVERRIDE;
+	virtual void CustomizeHeader(TSharedRef<IPropertyHandle> StructPropertyHandle, class FDetailWidgetRow& HeaderRow, IPropertyTypeCustomizationUtils& StructCustomizationUtils) override;
+	virtual void CustomizeChildren(TSharedRef<IPropertyHandle> StructPropertyHandle, class IDetailChildrenBuilder& ChildBuilder, IPropertyTypeCustomizationUtils& StructCustomizationUtils) override;
 	// End of IPropertyTypeCustomization interface
 
 	FText GetMinValue(TSharedRef<class IPropertyHandle> StructPropertyHandle) const;
@@ -61,8 +61,8 @@ public:
 	static TSharedRef<IPropertyTypeCustomization> MakeInstance();
 
 	// IPropertyTypeCustomization interface
-	virtual void CustomizeHeader(TSharedRef<IPropertyHandle> StructPropertyHandle, class FDetailWidgetRow& HeaderRow, IPropertyTypeCustomizationUtils& StructCustomizationUtils) OVERRIDE;
-	virtual void CustomizeChildren(TSharedRef<IPropertyHandle> StructPropertyHandle, class IDetailChildrenBuilder& ChildBuilder, IPropertyTypeCustomizationUtils& StructCustomizationUtils) OVERRIDE;
+	virtual void CustomizeHeader(TSharedRef<IPropertyHandle> StructPropertyHandle, class FDetailWidgetRow& HeaderRow, IPropertyTypeCustomizationUtils& StructCustomizationUtils) override;
+	virtual void CustomizeChildren(TSharedRef<IPropertyHandle> StructPropertyHandle, class IDetailChildrenBuilder& ChildBuilder, IPropertyTypeCustomizationUtils& StructCustomizationUtils) override;
 
 private:
 	// Storage object for bone hierarchy
diff --git a/Engine/Source/Editor/Persona/Private/AnimationEditorViewportClient.h b/Engine/Source/Editor/Persona/Private/AnimationEditorViewportClient.h
--- a/Engine/Source/Editor/Persona/Private/AnimationEditorViewportClient.h
+++ b/Engine/Source/Editor/Persona/Private/AnimationEditorViewportClient.h
@@ -39,25 +39,25 @@ public:
 	virtual ~FAnimationViewportClient();
 
 	// FEditorViewportClient interface
-	virtual FLinearColor GetBackgroundColor() const OVERRIDE;
-	virtual void Tick(float DeltaSeconds) OVERRIDE;
-	virtual void Draw(const FSceneView* View,FPrimitiveDrawInterface* PDI) OVERRIDE;
-	virtual void DrawCanvas( FViewport& InViewport, FSceneView& View, FCanvas& Canvas ) OVERRIDE;
-	virtual FSceneView* CalcSceneView(FSceneViewFamily* ViewFamily) OVERRIDE;
-	virtual bool InputKey(FViewport* Viewport, int32 ControllerId, FKey Key, EInputEvent Event, float AmountDepressed = 1.f, bool bGamepad=false) OVERRIDE;
-	virtual void ProcessClick(class FSceneView& View, class HHitProxy* HitProxy, FKey Key, EInputEvent Event, uint32 HitX, uint32 HitY) OVERRIDE;
-	virtual bool InputWidgetDelta( FViewport* Viewport, EAxisList::Type CurrentAxis, FVector& Drag, FRotator& Rot, FVector& Scale ) OVERRIDE;
-	virtual void TrackingStarted( const struct FInputEventState& InInputState, bool bIsDragging, bool bNudge ) OVERRIDE;
-	virtual void TrackingStopped() OVERRIDE;
-	virtual FWidget::EWidgetMode GetWidgetMode() const OVERRIDE;
-	virtual void SetWidgetMode(FWidget::EWidgetMode InWidgetMode) OVERRIDE;
-	virtual bool CanSetWidgetMode(FWidget::EWidgetMode NewMode) const OVERRIDE;
-	virtual FVector GetWidgetLocation() const OVERRIDE;
-	virtual FMatrix GetWidgetCoordSystem() const OVERRIDE;
-	virtual ECoordSystem GetWidgetCoordSystemSpace() const OVERRIDE;
-	virtual void SetWidgetCoordSystemSpace(ECoordSystem NewCoordSystem) OVERRIDE;
-	virtual void SetViewMode(EViewModeIndex InViewModeIndex) OVERRIDE;
-	virtual void SetViewportType(ELevelViewportType InViewportType) OVERRIDE;
+	virtual FLinearColor GetBackgroundColor() const override;
+	virtual void Tick(float DeltaSeconds) override;
+	virtual void Draw(const FSceneView* View,FPrimitiveDrawInterface* PDI) override;
+	virtual void DrawCanvas( FViewport& InViewport, FSceneView& View, FCanvas& Canvas ) override;
+	virtual FSceneView* CalcSceneView(FSceneViewFamily* ViewFamily) override;
+	virtual bool InputKey(FViewport* Viewport, int32 ControllerId, FKey Key, EInputEvent Event, float AmountDepressed = 1.f, bool bGamepad=false) override;
+	virtual void ProcessClick(class FSceneView& View, class HHitProxy* HitProxy, FKey Key, EInputEvent Event, uint32 HitX, uint32 HitY) override;
+	virtual bool InputWidgetDelta( FViewport* Viewport, EAxisList::Type CurrentAxis, FVector& Drag, FRotator& Rot, FVector& Scale ) override;
+	virtual void TrackingStarted( const struct FInputEventState& InInputState, bool bIsDragging, bool bNudge ) override;
+	virtual void TrackingStopped() override;
+	virtual FWidget::EWidgetMode GetWidgetMode() const override;
+	virtual void SetWidgetMode(FWidget::EWidgetMode InWidgetMode) override;
+	virtual bool CanSetWidgetMode(FWidget::EWidgetMode NewMode) const override;
+	virtual FVector GetWidgetLocation() const override;
+	virtual FMatrix GetWidgetCoordSystem() const override;
+	virtual ECoordSystem GetWidgetCoordSystemSpace() const override;
+	virtual void SetWidgetCoordSystemSpace(ECoordSystem NewCoordSystem) override;
+	virtual void SetViewMode(EViewModeIndex InViewModeIndex) override;
+	virtual void SetViewportType(ELevelViewportType InViewportType) override;
 	// End of FEditorViewportClient interface
 
 	/** Draw call to render UV overlay */
diff --git a/Engine/Source/Editor/Persona/Private/AnimationMode/AnimationMode.cpp b/Engine/Source/Editor/Persona/Private/AnimationMode/AnimationMode.cpp
--- a/Engine/Source/Editor/Persona/Private/AnimationMode/AnimationMode.cpp
+++ b/Engine/Source/Editor/Persona/Private/AnimationMode/AnimationMode.cpp
@@ -49,12 +49,12 @@ public:
 	}
 
 	// SSingleObjectDetailsPanel interface
-	virtual UObject* GetObjectToObserve() const OVERRIDE
+	virtual UObject* GetObjectToObserve() const override
 	{
 		return PersonaPtr.Pin()->GetAnimationAssetBeingEdited();
 	}
 
-	virtual TSharedRef<SWidget> PopulateSlot(TSharedRef<SWidget> PropertyEditorWidget) OVERRIDE
+	virtual TSharedRef<SWidget> PopulateSlot(TSharedRef<SWidget> PropertyEditorWidget) override
 	{
 		return SNew(SVerticalBox)
 			+SVerticalBox::Slot()
@@ -96,10 +96,10 @@ struct FAnimAssetPropertiesSummoner : public FWorkflowTabFactory
 	FAnimAssetPropertiesSummoner(TSharedPtr<class FAssetEditorToolkit> InHostingApp);
 
 	// FWorkflowTabFactory interface
-	virtual TSharedRef<SWidget> CreateTabBody(const FWorkflowTabSpawnInfo& Info) const OVERRIDE;
+	virtual TSharedRef<SWidget> CreateTabBody(const FWorkflowTabSpawnInfo& Info) const override;
 
 	// Create a tooltip widget for the tab
-	virtual TSharedPtr<SToolTip> CreateTabToolTipWidget(const FWorkflowTabSpawnInfo& Info) const OVERRIDE
+	virtual TSharedPtr<SToolTip> CreateTabToolTipWidget(const FWorkflowTabSpawnInfo& Info) const override
 	{
 		return  IDocumentation::Get()->CreateToolTip(LOCTEXT("AnimAssetPropertiesTooltip", "The Anim Asset Details tab lets you edit properties of the selection animation asset (animation, blend space etc)."), NULL, TEXT("Shared/Editors/Persona"), TEXT("AnimationAssetDetail_Window"));
 	}
diff --git a/Engine/Source/Editor/Persona/Private/BlueprintMode/AnimBlueprintMode.h b/Engine/Source/Editor/Persona/Private/BlueprintMode/AnimBlueprintMode.h
--- a/Engine/Source/Editor/Persona/Private/BlueprintMode/AnimBlueprintMode.h
+++ b/Engine/Source/Editor/Persona/Private/BlueprintMode/AnimBlueprintMode.h
@@ -15,7 +15,7 @@ public:
 	FAnimBlueprintEditAppMode(TSharedPtr<FPersona> InPersona);
 
 	// FApplicationMode interface
-	virtual void RegisterTabFactories(TSharedPtr<FTabManager> InTabManager) OVERRIDE;
-	virtual void PostActivateMode() OVERRIDE;
+	virtual void RegisterTabFactories(TSharedPtr<FTabManager> InTabManager) override;
+	virtual void PostActivateMode() override;
 	// End of FApplicationMode interface
 };
diff --git a/Engine/Source/Editor/Persona/Private/Customization/AnimGraphNodeSlotDetails.h b/Engine/Source/Editor/Persona/Private/Customization/AnimGraphNodeSlotDetails.h
--- a/Engine/Source/Editor/Persona/Private/Customization/AnimGraphNodeSlotDetails.h
+++ b/Engine/Source/Editor/Persona/Private/Customization/AnimGraphNodeSlotDetails.h
@@ -17,7 +17,7 @@ public:
 	}
 
 	// IDetailCustomization interface
-	virtual void CustomizeDetails(class IDetailLayoutBuilder& DetailBuilder) OVERRIDE;
+	virtual void CustomizeDetails(class IDetailLayoutBuilder& DetailBuilder) override;
 	
 private:
 	// persona editor 
diff --git a/Engine/Source/Editor/Persona/Private/EditorObjectsTracker.h b/Engine/Source/Editor/Persona/Private/EditorObjectsTracker.h
--- a/Engine/Source/Editor/Persona/Private/EditorObjectsTracker.h
+++ b/Engine/Source/Editor/Persona/Private/EditorObjectsTracker.h
@@ -10,7 +10,7 @@ class FEditorObjectTracker : public FGCObject
 public:
 
 	// FGCObject interface
-	void AddReferencedObjects( FReferenceCollector& Collector ) OVERRIDE;
+	void AddReferencedObjects( FReferenceCollector& Collector ) override;
 	// End of FGCObject interface
 
 	/** Returns an existing editor object for the specified class or creates one
diff --git a/Engine/Source/Editor/Persona/Private/MeshMode/MeshMode.cpp b/Engine/Source/Editor/Persona/Private/MeshMode/MeshMode.cpp
--- a/Engine/Source/Editor/Persona/Private/MeshMode/MeshMode.cpp
+++ b/Engine/Source/Editor/Persona/Private/MeshMode/MeshMode.cpp
@@ -37,7 +37,7 @@ public:
 	}
 
 	// SSingleObjectDetailsPanel interface
-	virtual UObject* GetObjectToObserve() const OVERRIDE
+	virtual UObject* GetObjectToObserve() const override
 	{
 		return PersonaPtr.Pin()->GetMesh();
 	}
diff --git a/Engine/Source/Editor/Persona/Private/MeshMode/MeshMode.h b/Engine/Source/Editor/Persona/Private/MeshMode/MeshMode.h
--- a/Engine/Source/Editor/Persona/Private/MeshMode/MeshMode.h
+++ b/Engine/Source/Editor/Persona/Private/MeshMode/MeshMode.h
@@ -14,10 +14,10 @@ struct FMeshPropertiesSummoner : public FWorkflowTabFactory
 	FMeshPropertiesSummoner(TSharedPtr<class FAssetEditorToolkit> InHostingApp);
 
 	// FWorkflowTabFactory interface
-	virtual TSharedRef<SWidget> CreateTabBody(const FWorkflowTabSpawnInfo& Info) const OVERRIDE;
+	virtual TSharedRef<SWidget> CreateTabBody(const FWorkflowTabSpawnInfo& Info) const override;
 
 	// Create a tooltip widget for the tab
-	virtual TSharedPtr<SToolTip> CreateTabToolTipWidget(const FWorkflowTabSpawnInfo& Info) const OVERRIDE
+	virtual TSharedPtr<SToolTip> CreateTabToolTipWidget(const FWorkflowTabSpawnInfo& Info) const override
 	{
 		return  IDocumentation::Get()->CreateToolTip(LOCTEXT("MeshDetailsTooltip", "The Mesh Details tab lets you edit properties (materials etc) of the current Skeletal Mesh."), NULL, TEXT("Shared/Editors/Persona"), TEXT("MeshDetail_Window"));
 	}
diff --git a/Engine/Source/Editor/Persona/Private/Persona.cpp b/Engine/Source/Editor/Persona/Private/Persona.cpp
--- a/Engine/Source/Editor/Persona/Private/Persona.cpp
+++ b/Engine/Source/Editor/Persona/Private/Persona.cpp
@@ -2745,7 +2745,7 @@ static class FMeshHierarchyCmd : private FSelfRegisteringExec
 {
 public:
 	/** Console commands, see embeded usage statement **/
-	virtual bool Exec( UWorld*, const TCHAR* Cmd, FOutputDevice& Ar ) OVERRIDE
+	virtual bool Exec( UWorld*, const TCHAR* Cmd, FOutputDevice& Ar ) override
 	{
 		bool bResult = false;
 		if(FParse::Command(&Cmd,TEXT("TMH")))
diff --git a/Engine/Source/Editor/Persona/Private/Persona.h b/Engine/Source/Editor/Persona/Private/Persona.h
--- a/Engine/Source/Editor/Persona/Private/Persona.h
+++ b/Engine/Source/Editor/Persona/Private/Persona.h
@@ -84,7 +84,7 @@ public:
 
 	void SetDetailObject(UObject* Obj);
 
-	virtual void OnActiveTabChanged(TSharedPtr<SDockTab> PreviouslyActive, TSharedPtr<SDockTab> NewlyActivated) OVERRIDE;
+	virtual void OnActiveTabChanged(TSharedPtr<SDockTab> PreviouslyActive, TSharedPtr<SDockTab> NewlyActivated) override;
 
 	/** handle after compile **/
 	UDebugSkelMeshComponent* GetPreviewMeshComponent();
@@ -192,25 +192,25 @@ public:
 
 public:
 	// IToolkit interface
-	virtual FName GetToolkitFName() const OVERRIDE;
-	virtual FText GetBaseToolkitName() const OVERRIDE;
-	virtual FText GetToolkitName() const OVERRIDE;
-	virtual FString GetWorldCentricTabPrefix() const OVERRIDE;
-	virtual FLinearColor GetWorldCentricTabColorScale() const OVERRIDE;
-	virtual void SaveAsset_Execute() OVERRIDE;
+	virtual FName GetToolkitFName() const override;
+	virtual FText GetBaseToolkitName() const override;
+	virtual FText GetToolkitName() const override;
+	virtual FString GetWorldCentricTabPrefix() const override;
+	virtual FLinearColor GetWorldCentricTabColorScale() const override;
+	virtual void SaveAsset_Execute() override;
 	// End of IToolkit interface
 
 	/** @return the documentation location for this editor */
-	virtual FString GetDocumentationLink() const OVERRIDE
+	virtual FString GetDocumentationLink() const override
 	{
 		return FString(TEXT("Engine/Animation/Persona"));
 	}
 	
 	/** Returns a pointer to the Blueprint object we are currently editing, as long as we are editing exactly one */
-	virtual UBlueprint* GetBlueprintObj() const OVERRIDE;
+	virtual UBlueprint* GetBlueprintObj() const override;
 
 	// FTickableEditorObject interface
-	virtual void Tick(float DeltaTime) OVERRIDE;
+	virtual void Tick(float DeltaTime) override;
 	// End of FTickableEditorObject interface
 
 	/** Returns the image brush to use for each modes dirty marker */
@@ -221,41 +221,41 @@ public:
 protected:
 	// FBlueprintEditor interface
 	//virtual void CreateDefaultToolbar() OVERRIDE;
-	virtual void CreateDefaultCommands() OVERRIDE;
-	virtual void OnSelectBone() OVERRIDE;
-	virtual bool CanSelectBone() const OVERRIDE;
-	virtual void OnAddPosePin() OVERRIDE;
-	virtual bool CanAddPosePin() const OVERRIDE;
-	virtual void OnRemovePosePin() OVERRIDE;
-	virtual bool CanRemovePosePin() const OVERRIDE;
-	virtual void StartEditingDefaults(bool bAutoFocus, bool bForceRefresh = false) OVERRIDE;
-	virtual void Compile() OVERRIDE;
-	virtual void OnGraphEditorFocused(const TSharedRef<class SGraphEditor>& InGraphEditor) OVERRIDE;
-	virtual FString GetDefaultEditorTitle() OVERRIDE;
-	virtual void OnConvertToSequenceEvaluator() OVERRIDE;
-	virtual void OnConvertToSequencePlayer() OVERRIDE;
-	virtual void OnConvertToBlendSpaceEvaluator() OVERRIDE;
-	virtual void OnConvertToBlendSpacePlayer() OVERRIDE;
-	virtual bool IsInAScriptingMode() const OVERRIDE;
-	virtual void OnOpenRelatedAsset() OVERRIDE;
-	virtual void GetCustomDebugObjects(TArray<FCustomDebugObject>& DebugList) const OVERRIDE;
-	virtual void CreateDefaultTabContents(const TArray<UBlueprint*>& InBlueprints) OVERRIDE;
-	virtual FGraphAppearanceInfo GetGraphAppearance() const OVERRIDE;
-	virtual bool IsEditable(UEdGraph* InGraph) const OVERRIDE;
-	virtual FString GetGraphDecorationString(UEdGraph* InGraph) const OVERRIDE;
+	virtual void CreateDefaultCommands() override;
+	virtual void OnSelectBone() override;
+	virtual bool CanSelectBone() const override;
+	virtual void OnAddPosePin() override;
+	virtual bool CanAddPosePin() const override;
+	virtual void OnRemovePosePin() override;
+	virtual bool CanRemovePosePin() const override;
+	virtual void StartEditingDefaults(bool bAutoFocus, bool bForceRefresh = false) override;
+	virtual void Compile() override;
+	virtual void OnGraphEditorFocused(const TSharedRef<class SGraphEditor>& InGraphEditor) override;
+	virtual FString GetDefaultEditorTitle() override;
+	virtual void OnConvertToSequenceEvaluator() override;
+	virtual void OnConvertToSequencePlayer() override;
+	virtual void OnConvertToBlendSpaceEvaluator() override;
+	virtual void OnConvertToBlendSpacePlayer() override;
+	virtual bool IsInAScriptingMode() const override;
+	virtual void OnOpenRelatedAsset() override;
+	virtual void GetCustomDebugObjects(TArray<FCustomDebugObject>& DebugList) const override;
+	virtual void CreateDefaultTabContents(const TArray<UBlueprint*>& InBlueprints) override;
+	virtual FGraphAppearanceInfo GetGraphAppearance() const override;
+	virtual bool IsEditable(UEdGraph* InGraph) const override;
+	virtual FString GetGraphDecorationString(UEdGraph* InGraph) const override;
 	// End of FBlueprintEditor interface
 
 	// IAssetEditorInstance interface
-	virtual void FocusWindow(UObject* ObjectToFocusOn = NULL) OVERRIDE;
+	virtual void FocusWindow(UObject* ObjectToFocusOn = NULL) override;
 	// End of IAssetEditorInstance interface
 
 	// Begin FEditorUndoClient Interface
-	virtual void PostUndo(bool bSuccess) OVERRIDE;
-	virtual void PostRedo(bool bSuccess) OVERRIDE;
+	virtual void PostUndo(bool bSuccess) override;
+	virtual void PostRedo(bool bSuccess) override;
 	// End of FEditorUndoClient
 
 	// Begin FAssetEditorToolkit Interface
-	virtual void FindInContentBrowser_Execute() OVERRIDE;
+	virtual void FindInContentBrowser_Execute() override;
 	// End of FAssetEditorToolkit
 
 	// Generic Command handlers
diff --git a/Engine/Source/Editor/Persona/Private/SAnimCompositeEditor.h b/Engine/Source/Editor/Persona/Private/SAnimCompositeEditor.h
--- a/Engine/Source/Editor/Persona/Private/SAnimCompositeEditor.h
+++ b/Engine/Source/Editor/Persona/Private/SAnimCompositeEditor.h
@@ -49,14 +49,14 @@ private:
 	void CollapseComposite();
 
 protected:
-	virtual void InitDetailsViewEditorObject(class UEditorAnimBaseObj* EdObj) OVERRIDE;
+	virtual void InitDetailsViewEditorObject(class UEditorAnimBaseObj* EdObj) override;
 
 public:
 	void Construct(const FArguments& InArgs);
 	
 	/** Return the animation composite being edited */
 	UAnimComposite* GetCompositeObj() const { return CompositeObj; }
-	virtual UAnimSequenceBase* GetEditorObject() const OVERRIDE { return GetCompositeObj(); }
+	virtual UAnimSequenceBase* GetEditorObject() const override { return GetCompositeObj(); }
 
 private:
 	/** Pointer to the animation composite being edited */
@@ -64,7 +64,7 @@ private:
 	
 	void PostUndo();
 	
-	virtual float CalculateSequenceLengthOfEditorObject() const OVERRIDE;
+	virtual float CalculateSequenceLengthOfEditorObject() const override;
 
 	/** This will sort all components of the montage and update (recreate) the UI */
 	void SortAndUpdateComposite();
@@ -76,6 +76,6 @@ public:
 	void			PostAnimUpdate();
 
 	// SWidget interface start
-	virtual void Tick( const FGeometry& AllottedGeometry, const double InCurrentTime, const float InDeltaTime ) OVERRIDE;
+	virtual void Tick( const FGeometry& AllottedGeometry, const double InCurrentTime, const float InDeltaTime ) override;
 	// SWidget interface end
 };
diff --git a/Engine/Source/Editor/Persona/Private/SAnimCurvePanel.cpp b/Engine/Source/Editor/Persona/Private/SAnimCurvePanel.cpp
--- a/Engine/Source/Editor/Persona/Private/SAnimCurvePanel.cpp
+++ b/Engine/Source/Editor/Persona/Private/SAnimCurvePanel.cpp
@@ -112,12 +112,12 @@ public:
 
 protected:
 	// SWidget interface
-	virtual int32 OnPaint(const FGeometry& AllottedGeometry, const FSlateRect& MyClippingRect, FSlateWindowElementList& OutDrawElements, int32 LayerId, const FWidgetStyle& InWidgetStyle, bool bParentEnabled) const OVERRIDE;
+	virtual int32 OnPaint(const FGeometry& AllottedGeometry, const FSlateRect& MyClippingRect, FSlateWindowElementList& OutDrawElements, int32 LayerId, const FWidgetStyle& InWidgetStyle, bool bParentEnabled) const override;
 	// SWidget interface
 
 	// SCurveEditor interface
-	virtual void SetDefaultOutput(const float MinZoomRange) OVERRIDE;
-	virtual float GetTimeStep(FTrackScaleInfo &ScaleInfo) const OVERRIDE;
+	virtual void SetDefaultOutput(const float MinZoomRange) override;
+	virtual float GetTimeStep(FTrackScaleInfo &ScaleInfo) const override;
 	// SCurveEditor interface
 	
 private:
diff --git a/Engine/Source/Editor/Persona/Private/SAnimEditorBase.h b/Engine/Source/Editor/Persona/Private/SAnimEditorBase.h
--- a/Engine/Source/Editor/Persona/Private/SAnimEditorBase.h
+++ b/Engine/Source/Editor/Persona/Private/SAnimEditorBase.h
@@ -62,7 +62,7 @@ public:
 	void ClearDetailsView();
 
 	// FGCObject interface start
-	virtual void AddReferencedObjects( FReferenceCollector& Collector ) OVERRIDE;
+	virtual void AddReferencedObjects( FReferenceCollector& Collector ) override;
 	// FGCObject interface end
 
 protected:
diff --git a/Engine/Source/Editor/Persona/Private/SAnimMontagePanel.cpp b/Engine/Source/Editor/Persona/Private/SAnimMontagePanel.cpp
--- a/Engine/Source/Editor/Persona/Private/SAnimMontagePanel.cpp
+++ b/Engine/Source/Editor/Persona/Private/SAnimMontagePanel.cpp
@@ -37,7 +37,7 @@ public:
 		return Operation;
 	}
 
-	virtual void OnDragged( const class FDragDropEvent& DragDropEvent ) OVERRIDE;
+	virtual void OnDragged( const class FDragDropEvent& DragDropEvent ) override;
 };
 
 class SMontageBranchingPointNode : public STrackNode
@@ -67,11 +67,11 @@ public:
 	SLATE_END_ARGS()
 
 	void Construct(const FArguments& InArgs);
-	virtual int32 OnPaint(const FGeometry& AllottedGeometry, const FSlateRect& MyClippingRect, FSlateWindowElementList& OutDrawElements, int32 LayerId, const FWidgetStyle& InWidgetStyle, bool bParentEnabled) const OVERRIDE;
+	virtual int32 OnPaint(const FGeometry& AllottedGeometry, const FSlateRect& MyClippingRect, FSlateWindowElementList& OutDrawElements, int32 LayerId, const FWidgetStyle& InWidgetStyle, bool bParentEnabled) const override;
 
-	virtual FVector2D GetDragDropScreenSpacePosition(const FGeometry& ParentAllottedGeometry, const FDragDropEvent& DragDropEvent) const OVERRIDE;
+	virtual FVector2D GetDragDropScreenSpacePosition(const FGeometry& ParentAllottedGeometry, const FDragDropEvent& DragDropEvent) const override;
 
-	virtual FReply BeginDrag( const FGeometry& MyGeometry, const FPointerEvent& MouseEvent ) OVERRIDE;
+	virtual FReply BeginDrag( const FGeometry& MyGeometry, const FPointerEvent& MouseEvent ) override;
 
 	bool SnapToDragBars() const {return true;}
 
@@ -85,11 +85,11 @@ private:
 	FBranchingPoint* BranchingPoint;
 
 	FString GetBranchingPointText() const;
-	virtual FVector2D GetSize() const OVERRIDE;
+	virtual FVector2D GetSize() const override;
 
 	// virtual draw related functions STrackNode
-	virtual FVector2D GetOffsetRelativeToParent(const FGeometry& ParentAllottedGeometry) const OVERRIDE;
-	virtual FVector2D GetSizeRelativeToParent(const FGeometry& ParentAllottedGeometry) const OVERRIDE;
+	virtual FVector2D GetOffsetRelativeToParent(const FGeometry& ParentAllottedGeometry) const override;
+	virtual FVector2D GetSizeRelativeToParent(const FGeometry& ParentAllottedGeometry) const override;
 };
 
 void FBranchingPointDragDropOp::OnDragged( const class FDragDropEvent& DragDropEvent )
diff --git a/Engine/Source/Editor/Persona/Private/SAnimMontagePanel.h b/Engine/Source/Editor/Persona/Private/SAnimMontagePanel.h
--- a/Engine/Source/Editor/Persona/Private/SAnimMontagePanel.h
+++ b/Engine/Source/Editor/Persona/Private/SAnimMontagePanel.h
@@ -59,7 +59,7 @@ public:
 	void SetMontage(class UAnimMontage * InMontage);
 
 	// SWidget interface
-	virtual FReply OnMouseButtonDown( const FGeometry& MyGeometry, const FPointerEvent& MouseEvent ) OVERRIDE;
+	virtual FReply OnMouseButtonDown( const FGeometry& MyGeometry, const FPointerEvent& MouseEvent ) override;
 
 	void Update();
 
diff --git a/Engine/Source/Editor/Persona/Private/SAnimMontageScrubPanel.h b/Engine/Source/Editor/Persona/Private/SAnimMontageScrubPanel.h
--- a/Engine/Source/Editor/Persona/Private/SAnimMontageScrubPanel.h
+++ b/Engine/Source/Editor/Persona/Private/SAnimMontageScrubPanel.h
@@ -36,18 +36,18 @@ public:
 	void Construct( const FArguments& InArgs );
 protected:
 	// notifiers 
-	virtual FReply OnClick_Backward_End() OVERRIDE;
-	virtual FReply OnClick_Forward_End() OVERRIDE;
+	virtual FReply OnClick_Backward_End() override;
+	virtual FReply OnClick_Forward_End() override;
 
-	virtual FReply OnClick_Backward_Step() OVERRIDE;
-	virtual FReply OnClick_Forward_Step() OVERRIDE;
+	virtual FReply OnClick_Backward_Step() override;
+	virtual FReply OnClick_Forward_Step() override;
 
-	virtual FReply OnClick_Forward() OVERRIDE;
-	virtual FReply OnClick_Backward() OVERRIDE;
+	virtual FReply OnClick_Forward() override;
+	virtual FReply OnClick_Backward() override;
 
-	virtual FReply OnClick_ToggleLoop() OVERRIDE;
+	virtual FReply OnClick_ToggleLoop() override;
 
-	virtual void OnValueChanged(float NewValue) OVERRIDE;
+	virtual void OnValueChanged(float NewValue) override;
 
 private:
 	TWeakPtr<SMontageEditor> MontageEditor;
diff --git a/Engine/Source/Editor/Persona/Private/SAnimNotifyPanel.cpp b/Engine/Source/Editor/Persona/Private/SAnimNotifyPanel.cpp
--- a/Engine/Source/Editor/Persona/Private/SAnimNotifyPanel.cpp
+++ b/Engine/Source/Editor/Persona/Private/SAnimNotifyPanel.cpp
@@ -63,10 +63,10 @@ public:
 	void Construct(const FArguments& Declaration);
 
 	// SWidget interface
-	virtual FReply OnDragDetected( const FGeometry& MyGeometry, const FPointerEvent& MouseEvent) OVERRIDE;
-	virtual FReply OnMouseMove( const FGeometry& MyGeometry, const FPointerEvent& MouseEvent) OVERRIDE;
-	virtual FReply OnMouseButtonUp( const FGeometry& MyGeometry, const FPointerEvent& MouseEvent ) OVERRIDE;
-	virtual void Tick(const FGeometry& AllottedGeometry, const double InCurrentTime, const float InDeltaTime) OVERRIDE;
+	virtual FReply OnDragDetected( const FGeometry& MyGeometry, const FPointerEvent& MouseEvent) override;
+	virtual FReply OnMouseMove( const FGeometry& MyGeometry, const FPointerEvent& MouseEvent) override;
+	virtual FReply OnMouseButtonUp( const FGeometry& MyGeometry, const FPointerEvent& MouseEvent ) override;
+	virtual void Tick(const FGeometry& AllottedGeometry, const double InCurrentTime, const float InDeltaTime) override;
 	
 	// End of SWidget interface
 
@@ -84,8 +84,8 @@ public:
 	UObject* GetObjectBeingDisplayed() const;
 	// End of SNodePanel::SNode
 
-	virtual FVector2D ComputeDesiredSize() const  OVERRIDE;
-	virtual int32 OnPaint(const FGeometry& AllottedGeometry, const FSlateRect& MyClippingRect, FSlateWindowElementList& OutDrawElements, int32 LayerId, const FWidgetStyle& InWidgetStyle, bool bParentEnabled) const OVERRIDE;
+	virtual FVector2D ComputeDesiredSize() const  override;
+	virtual int32 OnPaint(const FGeometry& AllottedGeometry, const FSlateRect& MyClippingRect, FSlateWindowElementList& OutDrawElements, int32 LayerId, const FWidgetStyle& InWidgetStyle, bool bParentEnabled) const override;
 
 	/** Helpers to draw scrub handles and snap offsets */
 	void DrawHandleOffset( const float& Offset, const float& HandleCentre, FSlateWindowElementList& OutDrawElements, int32 MarkerLayer, const FGeometry &AllottedGeometry, const FSlateRect& MyClippingRect ) const;
@@ -229,7 +229,7 @@ public:
 
 	// SWidget interface
 	virtual void Tick( const FGeometry& AllottedGeometry, const double InCurrentTime, const float InDeltaTime ) {UpdateCachedGeometry(AllottedGeometry);}
-	virtual int32 OnPaint(const FGeometry& AllottedGeometry, const FSlateRect& MyClippingRect, FSlateWindowElementList& OutDrawElements, int32 LayerId, const FWidgetStyle& InWidgetStyle, bool bParentEnabled) const OVERRIDE;
+	virtual int32 OnPaint(const FGeometry& AllottedGeometry, const FSlateRect& MyClippingRect, FSlateWindowElementList& OutDrawElements, int32 LayerId, const FWidgetStyle& InWidgetStyle, bool bParentEnabled) const override;
 	// End of SWidget interface
 
 	/** Returns the cached rendering geometry of this track */
@@ -342,11 +342,11 @@ protected:
 	int32 GetHitNotifyNode(const FGeometry& MyGeometry, const FVector2D& Position);
 
 	TSharedPtr<SWidget> SummonContextMenu(const FGeometry& MyGeometry, const FPointerEvent& MouseEvent);
-	virtual FVector2D ComputeDesiredSize() const OVERRIDE;
-	virtual FReply OnMouseButtonUp(const FGeometry& MyGeometry, const FPointerEvent& MouseEvent) OVERRIDE;
+	virtual FVector2D ComputeDesiredSize() const override;
+	virtual FReply OnMouseButtonUp(const FGeometry& MyGeometry, const FPointerEvent& MouseEvent) override;
 
-	virtual FReply OnDrop( const FGeometry& MyGeometry, const FDragDropEvent& DragDropEvent ) OVERRIDE;
-	virtual FReply OnMouseButtonDown( const FGeometry& MyGeometry, const FPointerEvent& MouseEvent ) OVERRIDE;
+	virtual FReply OnDrop( const FGeometry& MyGeometry, const FDragDropEvent& DragDropEvent ) override;
+	virtual FReply OnMouseButtonDown( const FGeometry& MyGeometry, const FPointerEvent& MouseEvent ) override;
 
 	float CalculateTime( const FGeometry& MyGeometry, FVector2D NodePos, bool bInputIsAbsolute = true );
 
@@ -494,7 +494,7 @@ public:
 
 	DRAG_DROP_OPERATOR_TYPE(FNotifyDragDropOp, FDragDropOperation)
 
-	virtual void OnDrop( bool bDropWasHandled, const FPointerEvent& MouseEvent ) OVERRIDE
+	virtual void OnDrop( bool bDropWasHandled, const FPointerEvent& MouseEvent ) override
 	{
 		if ( bDropWasHandled == false )
 		{
@@ -509,7 +509,7 @@ public:
 		FDragDropOperation::OnDrop(bDropWasHandled, MouseEvent);
 	}
 
-	virtual void OnDragged( const class FDragDropEvent& DragDropEvent ) OVERRIDE
+	virtual void OnDragged( const class FDragDropEvent& DragDropEvent ) override
 	{
 		// Reset snapped node pointer
 		SnappedNode = NULL;
@@ -754,7 +754,7 @@ public:
 	}
 	
 	/** The widget decorator to use */
-	virtual TSharedPtr<SWidget> GetDefaultDecorator() const OVERRIDE
+	virtual TSharedPtr<SWidget> GetDefaultDecorator() const override
 	{
 		return Decorator;
 	}
@@ -787,7 +787,7 @@ public:
 	{
 	}
 
-	virtual void RegisterCommands() OVERRIDE
+	virtual void RegisterCommands() override
 	{
 		UI_COMMAND( SomeSequenceAction, "Some Sequence Action", "Does some sequence action", EUserInterfaceActionType::Button, FInputGesture() );
 	}
diff --git a/Engine/Source/Editor/Persona/Private/SAnimNotifyPanel.h b/Engine/Source/Editor/Persona/Private/SAnimNotifyPanel.h
--- a/Engine/Source/Editor/Persona/Private/SAnimNotifyPanel.h
+++ b/Engine/Source/Editor/Persona/Private/SAnimNotifyPanel.h
@@ -106,7 +106,7 @@ public:
 	/**Handler for when a notify node drag has been initiated */
 	FReply OnNotifyNodeDragStarted(TArray<TSharedPtr<SAnimNotifyNode>> NotifyNodes, TSharedRef<SWidget> Decorator, const FVector2D& ScreenCursorPos, const FVector2D& ScreenNodePosition, const bool bDragOnMarker);
 
-	virtual float GetSequenceLength() const OVERRIDE {return Sequence->SequenceLength;}
+	virtual float GetSequenceLength() const override {return Sequence->SequenceLength;}
 
 	void CopySelectedNotifiesToClipboard() const;
 	void OnPasteNotifies(SAnimNotifyTrack* RequestTrack, float ClickTime, ENotifyPasteMode::Type PasteMode, ENotifyPasteMultipleMode::Type MultiplePasteType);
@@ -140,7 +140,7 @@ private:
 	void OnGenericDelete();
 
 	/** We support keyboard focus to detect when we should process key commands like delete */
-	virtual bool SupportsKeyboardFocus() const OVERRIDE
+	virtual bool SupportsKeyboardFocus() const override
 	{
 		return true;
 	}
diff --git a/Engine/Source/Editor/Persona/Private/SAnimTrackPanel.h b/Engine/Source/Editor/Persona/Private/SAnimTrackPanel.h
--- a/Engine/Source/Editor/Persona/Private/SAnimTrackPanel.h
+++ b/Engine/Source/Editor/Persona/Private/SAnimTrackPanel.h
@@ -86,10 +86,10 @@ public:
 	void Construct(const FArguments& InArgs);
 
 	// SWidget interface
-	virtual FReply OnMouseWheel( const FGeometry& MyGeometry, const FPointerEvent& MouseEvent ) OVERRIDE;
-	virtual FReply OnMouseButtonDown( const FGeometry& MyGeometry, const FPointerEvent& MouseEvent ) OVERRIDE;
-	virtual FReply OnMouseButtonUp( const FGeometry& MyGeometry, const FPointerEvent& MouseEvent ) OVERRIDE;
-	virtual FReply OnMouseMove( const FGeometry& MyGeometry, const FPointerEvent& MouseEvent ) OVERRIDE;
+	virtual FReply OnMouseWheel( const FGeometry& MyGeometry, const FPointerEvent& MouseEvent ) override;
+	virtual FReply OnMouseButtonDown( const FGeometry& MyGeometry, const FPointerEvent& MouseEvent ) override;
+	virtual FReply OnMouseButtonUp( const FGeometry& MyGeometry, const FPointerEvent& MouseEvent ) override;
+	virtual FReply OnMouseMove( const FGeometry& MyGeometry, const FPointerEvent& MouseEvent ) override;
 
 	void PanInputViewRange(int32 ScreenDelta, FVector2D ScreenViewSize);
 
diff --git a/Engine/Source/Editor/Persona/Private/SAnimationBlendSpace.h b/Engine/Source/Editor/Persona/Private/SAnimationBlendSpace.h
--- a/Engine/Source/Editor/Persona/Private/SAnimationBlendSpace.h
+++ b/Engine/Source/Editor/Persona/Private/SAnimationBlendSpace.h
@@ -481,7 +481,7 @@ public:
 	 *
 	 * @return The maximum layer ID attained by this widget or any of its children.
 	 */
-	virtual int32 OnPaint( const FGeometry& AllottedGeometry, const FSlateRect& MyClippingRect, FSlateWindowElementList& OutDrawElements, int32 LayerId, const FWidgetStyle& InWidgetStyle, bool bParentEnabled ) const OVERRIDE;
+	virtual int32 OnPaint( const FGeometry& AllottedGeometry, const FSlateRect& MyClippingRect, FSlateWindowElementList& OutDrawElements, int32 LayerId, const FWidgetStyle& InWidgetStyle, bool bParentEnabled ) const override;
 
 	/**
 	 * Clear all points
@@ -492,7 +492,7 @@ public:
 		BlendSpaceGrid.Reset();
 	}
 
-	virtual void ResampleData() OVERRIDE;
+	virtual void ResampleData() override;
 
 	/** 
 	 * Manually AddPoint
@@ -522,13 +522,13 @@ public:
 	/** 
 	 * Mapping function between WidgetPos and EditorPos
 	 */
-	virtual TOptional<FVector2D>	GetWidgetPosFromEditorPos(const FVector & EditorPos, const FSlateRect& WindowRect) const OVERRIDE;
-	virtual TOptional<FVector>		GetEditorPosFromWidgetPos(const FVector2D & WidgetPos, const FSlateRect& WindowRect) const OVERRIDE;
+	virtual TOptional<FVector2D>	GetWidgetPosFromEditorPos(const FVector & EditorPos, const FSlateRect& WindowRect) const override;
+	virtual TOptional<FVector>		GetEditorPosFromWidgetPos(const FVector2D & WidgetPos, const FSlateRect& WindowRect) const override;
 
 	/**
 	 * Snaps a position in editor space to the editor grid
 	 */
-	virtual FVector					SnapEditorPosToGrid(const FVector& InPos) const OVERRIDE;
+	virtual FVector					SnapEditorPosToGrid(const FVector& InPos) const override;
 
 	/**
 	 * Utility function to set the grid info for both the BlendSpaceGrid and the Generator in one go
@@ -561,13 +561,13 @@ private:
 	const UBlendSpace* GetBlendSpace() const { return Cast<UBlendSpace>(BlendSpace); }
 
 	/** Utility functions **/
-	virtual FText GetInputText(const FVector & GridPos) const OVERRIDE;
+	virtual FText GetInputText(const FVector & GridPos) const override;
 	bool IsOnTheBorder(const FIntPoint & GridIndices) const;
 
 	/** Utility function for OnPaint **/
 	void DrawHighlightGrid(const FVector2D & LeftTopPos, const FVector2D & RightBottomPos, const FGeometry& AllottedGeometry, const FSlateRect& MyClippingRect, FSlateWindowElementList& OutDrawElements, int32 LayerId ) const;
 
-	virtual FReply UpdateLastMousePosition( const FGeometry& MyGeometry, const FVector2D& ScreenSpacePosition, bool bClampToWindowRect = false, bool bSnapToGrid = false  ) OVERRIDE;
+	virtual FReply UpdateLastMousePosition( const FGeometry& MyGeometry, const FVector2D& ScreenSpacePosition, bool bClampToWindowRect = false, bool bSnapToGrid = false  ) override;
 	bool GetMinMaxFromGridPos(FVector GridPos, FVector2D & WindowLeftTop, FVector2D & WindowRightBottom, const FSlateRect& WindowRect) const;
 };
 
diff --git a/Engine/Source/Editor/Persona/Private/SAnimationBlendSpace1D.h b/Engine/Source/Editor/Persona/Private/SAnimationBlendSpace1D.h
--- a/Engine/Source/Editor/Persona/Private/SAnimationBlendSpace1D.h
+++ b/Engine/Source/Editor/Persona/Private/SAnimationBlendSpace1D.h
@@ -122,7 +122,7 @@ public:
 	 *
 	 * @return The maximum layer ID attained by this widget or any of its children.
 	 */
-	virtual int32 OnPaint( const FGeometry& AllottedGeometry, const FSlateRect& MyClippingRect, FSlateWindowElementList& OutDrawElements, int32 LayerId, const FWidgetStyle& InWidgetStyle, bool bParentEnabled ) const OVERRIDE;
+	virtual int32 OnPaint( const FGeometry& AllottedGeometry, const FSlateRect& MyClippingRect, FSlateWindowElementList& OutDrawElements, int32 LayerId, const FWidgetStyle& InWidgetStyle, bool bParentEnabled ) const override;
 
 	/**
 	 * A Panel's desired size in the space required to arrange of its children on the screen while respecting all of
@@ -130,9 +130,9 @@ public:
 	 *
 	 * @return The desired size.
 	 */
-	FVector2D ComputeDesiredSize() const OVERRIDE;
+	FVector2D ComputeDesiredSize() const override;
 
-	virtual void ResampleData() OVERRIDE;
+	virtual void ResampleData() override;
 
 	/** Preview point that is fed into the blend space to generate output */
 	TAttribute<FVector>	PreviewInput;
@@ -140,18 +140,18 @@ public:
 	/** 
 	 * Mapping function between WidgetPos and GridPos
 	 */
-	virtual TOptional<FVector2D>	GetWidgetPosFromEditorPos(const FVector & EditorPos, const FSlateRect& WindowRect) const OVERRIDE;
-	virtual TOptional<FVector>		GetEditorPosFromWidgetPos(const FVector2D & WidgetPos, const FSlateRect& WindowRect) const OVERRIDE;
+	virtual TOptional<FVector2D>	GetWidgetPosFromEditorPos(const FVector & EditorPos, const FSlateRect& WindowRect) const override;
+	virtual TOptional<FVector>		GetEditorPosFromWidgetPos(const FVector2D & WidgetPos, const FSlateRect& WindowRect) const override;
 
 	/**
 	 * Snaps a position in editor space to the editor grid
 	 */
-	virtual FVector					SnapEditorPosToGrid(const FVector& InPos) const OVERRIDE;
+	virtual FVector					SnapEditorPosToGrid(const FVector& InPos) const override;
 
 protected:
 	/** Utility functions **/
-	virtual FText GetInputText(const FVector & GridPos) const OVERRIDE;
-	virtual FReply UpdateLastMousePosition( const FGeometry& MyGeometry, const FVector2D& ScreenSpacePosition, bool bClampToWindowRect = false, bool bSnapToGrid = false  ) OVERRIDE;
+	virtual FText GetInputText(const FVector & GridPos) const override;
+	virtual FReply UpdateLastMousePosition( const FGeometry& MyGeometry, const FVector2D& ScreenSpacePosition, bool bClampToWindowRect = false, bool bSnapToGrid = false  ) override;
 
 private:
 
@@ -181,7 +181,7 @@ public:
 
 protected:
 
-	virtual TSharedRef<SWidget> MakeDisplayOptionsBox() const OVERRIDE;
+	virtual TSharedRef<SWidget> MakeDisplayOptionsBox() const override;
 
 private:
 
diff --git a/Engine/Source/Editor/Persona/Private/SAnimationBlendSpaceBase.cpp b/Engine/Source/Editor/Persona/Private/SAnimationBlendSpaceBase.cpp
--- a/Engine/Source/Editor/Persona/Private/SAnimationBlendSpaceBase.cpp
+++ b/Engine/Source/Editor/Persona/Private/SAnimationBlendSpaceBase.cpp
@@ -57,7 +57,7 @@ public:
 	TWeakPtr<SBlendSpaceWidget> BlendSpaceWidget;
 
 	/** The widget decorator to use */
-	virtual TSharedPtr<SWidget> GetDefaultDecorator() const OVERRIDE
+	virtual TSharedPtr<SWidget> GetDefaultDecorator() const override
 	{
 		return SNew(SBorder)
 			.BorderImage(FEditorStyle::GetBrush("Graph.ConnectorFeedback.Border"))
diff --git a/Engine/Source/Editor/Persona/Private/SAnimationBlendSpaceBase.h b/Engine/Source/Editor/Persona/Private/SAnimationBlendSpaceBase.h
--- a/Engine/Source/Editor/Persona/Private/SAnimationBlendSpaceBase.h
+++ b/Engine/Source/Editor/Persona/Private/SAnimationBlendSpaceBase.h
@@ -249,13 +249,13 @@ public:
 	FVector		GetLastValidMouseGridPoint() { return LastValidMouseEditorPoint; }
 
 	// SWidget Interface
-	virtual FReply	OnDrop( const FGeometry& MyGeometry, const FDragDropEvent& DragDropEvent ) OVERRIDE;
-	virtual FReply	OnDragDetected( const FGeometry& MyGeometry, const FPointerEvent& MouseEvent ) OVERRIDE;
-	virtual FReply	OnDragOver( const FGeometry& MyGeometry, const FDragDropEvent& DragDropEvent ) OVERRIDE;
-	virtual FReply	OnMouseButtonUp( const FGeometry& MyGeometry, const FPointerEvent& MouseEvent ) OVERRIDE;	
-	virtual FReply	OnMouseButtonDown( const FGeometry& MyGeometry, const FPointerEvent& MouseEvent ) OVERRIDE;
-	virtual FReply	OnMouseMove( const FGeometry& MyGeometry, const FPointerEvent& MouseEvent ) OVERRIDE;
-	virtual void	OnDragLeave( const FDragDropEvent& DragDropEvent ) OVERRIDE;
+	virtual FReply	OnDrop( const FGeometry& MyGeometry, const FDragDropEvent& DragDropEvent ) override;
+	virtual FReply	OnDragDetected( const FGeometry& MyGeometry, const FPointerEvent& MouseEvent ) override;
+	virtual FReply	OnDragOver( const FGeometry& MyGeometry, const FDragDropEvent& DragDropEvent ) override;
+	virtual FReply	OnMouseButtonUp( const FGeometry& MyGeometry, const FPointerEvent& MouseEvent ) override;	
+	virtual FReply	OnMouseButtonDown( const FGeometry& MyGeometry, const FPointerEvent& MouseEvent ) override;
+	virtual FReply	OnMouseMove( const FGeometry& MyGeometry, const FPointerEvent& MouseEvent ) override;
+	virtual void	OnDragLeave( const FDragDropEvent& DragDropEvent ) override;
 	// End of SWidget
 
 	/** Called by OnDrop to actually process sample drop */
@@ -356,7 +356,7 @@ public:
 	void Construct(const FArguments& InArgs);
 
 	// SWidget interface
-	virtual void	Tick(const FGeometry& AllottedGeometry, const double InCurrentTime, const float InDeltaTime) OVERRIDE;
+	virtual void	Tick(const FGeometry& AllottedGeometry, const double InCurrentTime, const float InDeltaTime) override;
 	// End of SWidget
 
 	void PostUndo();
diff --git a/Engine/Source/Editor/Persona/Private/SAnimationEditorViewport.h b/Engine/Source/Editor/Persona/Private/SAnimationEditorViewport.h
--- a/Engine/Source/Editor/Persona/Private/SAnimationEditorViewport.h
+++ b/Engine/Source/Editor/Persona/Private/SAnimationEditorViewport.h
@@ -41,8 +41,8 @@ public:
 
 protected:
 	// SEditorViewport interface
-	virtual TSharedRef<FEditorViewportClient> MakeEditorViewportClient() OVERRIDE;
-	virtual TSharedPtr<SWidget> MakeViewportToolbar() OVERRIDE;
+	virtual TSharedRef<FEditorViewportClient> MakeEditorViewportClient() override;
+	virtual TSharedPtr<SWidget> MakeViewportToolbar() override;
 	// End of SEditorViewport interface
 
 	/**  Handle undo/redo by refreshing the viewport */
@@ -91,7 +91,7 @@ public:
 	SAnimationEditorViewportTabBody();
 	virtual ~SAnimationEditorViewportTabBody();
 
-	virtual void Tick(const FGeometry& AllottedGeometry, const double InCurrentTime, const float InDeltaTime) OVERRIDE;
+	virtual void Tick(const FGeometry& AllottedGeometry, const double InCurrentTime, const float InDeltaTime) override;
 	void RefreshViewport();
 
 	/**
diff --git a/Engine/Source/Editor/Persona/Private/SAnimationSequenceBrowser.cpp b/Engine/Source/Editor/Persona/Private/SAnimationSequenceBrowser.cpp
--- a/Engine/Source/Editor/Persona/Private/SAnimationSequenceBrowser.cpp
+++ b/Engine/Source/Editor/Persona/Private/SAnimationSequenceBrowser.cpp
@@ -26,12 +26,12 @@ public:
 	FFrontendFilter_AdditiveAnimAssets(TSharedPtr<FFrontendFilterCategory> InCategory) : FFrontendFilter(InCategory) {}
 
 	// FFrontendFilter implementation
-	virtual FString GetName() const OVERRIDE { return TEXT("AdditiveAnimAssets"); }
-	virtual FText GetDisplayName() const OVERRIDE { return LOCTEXT("FFrontendFilter_AdditiveAnimAssets", "Additive Animations"); }
-	virtual FText GetToolTipText() const OVERRIDE { return LOCTEXT("FFrontendFilter_AdditiveAnimAssetsToolTip", "Show only animations that are additive."); }
+	virtual FString GetName() const override { return TEXT("AdditiveAnimAssets"); }
+	virtual FText GetDisplayName() const override { return LOCTEXT("FFrontendFilter_AdditiveAnimAssets", "Additive Animations"); }
+	virtual FText GetToolTipText() const override { return LOCTEXT("FFrontendFilter_AdditiveAnimAssetsToolTip", "Show only animations that are additive."); }
 
 	// IFilter implementation
-	virtual bool PassesFilter( AssetFilterType InItem ) const OVERRIDE
+	virtual bool PassesFilter( AssetFilterType InItem ) const override
 	{
 		FString TagValue = InItem.TagsAndValues.FindRef("AdditiveAnimType");
 		return !TagValue.IsEmpty() && !TagValue.Equals(TEXT("AAT_None"));
diff --git a/Engine/Source/Editor/Persona/Private/SMontageEditor.h b/Engine/Source/Editor/Persona/Private/SMontageEditor.h
--- a/Engine/Source/Editor/Persona/Private/SMontageEditor.h
+++ b/Engine/Source/Editor/Persona/Private/SMontageEditor.h
@@ -45,7 +45,7 @@ private:
 
 protected:
 	// Begin SAnimEditorBase interface
-	virtual TSharedRef<class SAnimationScrubPanel> ConstructAnimScrubPanel() OVERRIDE;
+	virtual TSharedRef<class SAnimationScrubPanel> ConstructAnimScrubPanel() override;
 	// End SAnimEditorBase interface
 
 public:
@@ -53,7 +53,7 @@ public:
 	void SetMontageObj(UAnimMontage * NewMontage);
 	UAnimMontage * GetMontageObj() const { return MontageObj; }
 
-	virtual UAnimSequenceBase* GetEditorObject() const OVERRIDE { return GetMontageObj(); }
+	virtual UAnimSequenceBase* GetEditorObject() const override { return GetMontageObj(); }
 
 	void RestartPreview();
 	void RestartPreviewFromSection(int32 FromSectionIdx = INDEX_NONE);
@@ -78,7 +78,7 @@ private:
 	/** If user is currently dragging an item */
 	bool bDragging;
 	
-	virtual float CalculateSequenceLengthOfEditorObject() const OVERRIDE;
+	virtual float CalculateSequenceLengthOfEditorObject() const override;
 	void SortAndUpdateMontage();
 	void CollapseMontage();
 	void SortBranchPoints();
@@ -86,7 +86,7 @@ private:
 	void SortSections();
 	void EnsureStartingSection();
 	void EnsureSlotNode();
-	virtual bool ClampToEndTime(float NewEndTime) OVERRIDE;
+	virtual bool ClampToEndTime(float NewEndTime) override;
 	void PostUndo();
 	
 	bool GetSectionTime( int32 SectionIndex, float &OutTime ) const;
@@ -99,7 +99,7 @@ private:
 	void RefreshNotifyTriggerOffsets();
 
 protected:
-	virtual void InitDetailsViewEditorObject(class UEditorAnimBaseObj* EdObj) OVERRIDE;
+	virtual void InitDetailsViewEditorObject(class UEditorAnimBaseObj* EdObj) override;
 
 public:
 
@@ -140,6 +140,6 @@ public:
 	void			PostAnimUpdate();
 
 	// SWidget interface start
-	virtual void Tick( const FGeometry& AllottedGeometry, const double InCurrentTime, const float InDeltaTime ) OVERRIDE;
+	virtual void Tick( const FGeometry& AllottedGeometry, const double InCurrentTime, const float InDeltaTime ) override;
 	// SWidget interface end
 };
diff --git a/Engine/Source/Editor/Persona/Private/SRetargetSourceManager.h b/Engine/Source/Editor/Persona/Private/SRetargetSourceManager.h
--- a/Engine/Source/Editor/Persona/Private/SRetargetSourceManager.h
+++ b/Engine/Source/Editor/Persona/Private/SRetargetSourceManager.h
@@ -166,7 +166,7 @@ public:
 	FText& GetFilterText() { return FilterText; }
 
 	// SWidget interface
-	virtual int32 OnPaint( const FGeometry& AllottedGeometry, const FSlateRect& MyClippingRect, FSlateWindowElementList& OutDrawElements, int32 LayerId, const FWidgetStyle& InWidgetStyle, bool bParentEnabled ) const OVERRIDE;
+	virtual int32 OnPaint( const FGeometry& AllottedGeometry, const FSlateRect& MyClippingRect, FSlateWindowElementList& OutDrawElements, int32 LayerId, const FWidgetStyle& InWidgetStyle, bool bParentEnabled ) const override;
 	// end of SWidget interface
 private:
 
diff --git a/Engine/Source/Editor/Persona/Private/SSequenceEditor.h b/Engine/Source/Editor/Persona/Private/SSequenceEditor.h
--- a/Engine/Source/Editor/Persona/Private/SSequenceEditor.h
+++ b/Engine/Source/Editor/Persona/Private/SSequenceEditor.h
@@ -38,7 +38,7 @@ public:
 	void SetSequenceObj(UAnimSequenceBase * NewSequence);
 	UAnimSequenceBase * GetSequenceObj() const { return SequenceObj; }
 
-	virtual UAnimSequenceBase* GetEditorObject() const OVERRIDE { return GetSequenceObj(); }
+	virtual UAnimSequenceBase* GetEditorObject() const override { return GetSequenceObj(); }
 
 private:
 	/** Pointer to the animation sequence being edited */
diff --git a/Engine/Source/Editor/Persona/Private/SSkeletonAnimNotifies.h b/Engine/Source/Editor/Persona/Private/SSkeletonAnimNotifies.h
--- a/Engine/Source/Editor/Persona/Private/SSkeletonAnimNotifies.h
+++ b/Engine/Source/Editor/Persona/Private/SSkeletonAnimNotifies.h
@@ -13,10 +13,10 @@ struct FSkeletonAnimNotifiesSummoner : public FWorkflowTabFactory
 public:
 	FSkeletonAnimNotifiesSummoner(TSharedPtr<class FAssetEditorToolkit> InHostingApp);
 
-	virtual TSharedRef<SWidget> CreateTabBody(const FWorkflowTabSpawnInfo& Info) const OVERRIDE;
+	virtual TSharedRef<SWidget> CreateTabBody(const FWorkflowTabSpawnInfo& Info) const override;
 
 	// Create a tooltip widget for the tab
-	virtual TSharedPtr<SToolTip> CreateTabToolTipWidget(const FWorkflowTabSpawnInfo& Info) const OVERRIDE
+	virtual TSharedPtr<SToolTip> CreateTabToolTipWidget(const FWorkflowTabSpawnInfo& Info) const override
 	{
 		return  IDocumentation::Get()->CreateToolTip(LOCTEXT("AnimationNotifierTooltip", "This tab lets you modify custom animation notifies"), NULL, TEXT("Shared/Editors/Persona"), TEXT("AnimationNotifies_Window"));
 	}
@@ -83,7 +83,7 @@ public:
 	void PostUndo();
 
 	// FGCObject interface start
-	virtual void AddReferencedObjects( FReferenceCollector& Collector ) OVERRIDE;
+	virtual void AddReferencedObjects( FReferenceCollector& Collector ) override;
 	// FGCObject interface end
 
 	/** When user attempts to commit the name of a track*/
diff --git a/Engine/Source/Editor/Persona/Private/SSkeletonSlotGroupNames.h b/Engine/Source/Editor/Persona/Private/SSkeletonSlotGroupNames.h
--- a/Engine/Source/Editor/Persona/Private/SSkeletonSlotGroupNames.h
+++ b/Engine/Source/Editor/Persona/Private/SSkeletonSlotGroupNames.h
@@ -13,10 +13,10 @@ struct FSkeletonSlotGroupNamesSummoner : public FWorkflowTabFactory
 public:
 	FSkeletonSlotGroupNamesSummoner(TSharedPtr<class FAssetEditorToolkit> InHostingApp);
 
-	virtual TSharedRef<SWidget> CreateTabBody(const FWorkflowTabSpawnInfo& Info) const OVERRIDE;
+	virtual TSharedRef<SWidget> CreateTabBody(const FWorkflowTabSpawnInfo& Info) const override;
 
 	// Create a tooltip widget for the tab
-	virtual TSharedPtr<SToolTip> CreateTabToolTipWidget(const FWorkflowTabSpawnInfo& Info) const OVERRIDE
+	virtual TSharedPtr<SToolTip> CreateTabToolTipWidget(const FWorkflowTabSpawnInfo& Info) const override
 	{
 		return  IDocumentation::Get()->CreateToolTip(LOCTEXT("WindowTooltip", "This tab lets you modify custom animation Slot Group Name"), NULL, TEXT("Shared/Editors/Persona"), TEXT("AnimationSlotGroupName_Window"));
 	}
@@ -83,7 +83,7 @@ public:
 	void PostUndo();
 
 	// FGCObject interface start
-	virtual void AddReferencedObjects( FReferenceCollector& Collector ) OVERRIDE;
+	virtual void AddReferencedObjects( FReferenceCollector& Collector ) override;
 	// FGCObject interface end
 
 	/** When user attempts to commit the name of a track*/
diff --git a/Engine/Source/Editor/Persona/Private/SSkeletonTree.cpp b/Engine/Source/Editor/Persona/Private/SSkeletonTree.cpp
--- a/Engine/Source/Editor/Persona/Private/SSkeletonTree.cpp
+++ b/Engine/Source/Editor/Persona/Private/SSkeletonTree.cpp
@@ -99,21 +99,21 @@ public:
 	void Construct( const FArguments& InArgs, const TSharedRef<STableViewBase>& InOwnerTableView );
 
 	/** Overridden from SMultiColumnTableRow.  Generates a widget for this column of the tree row. */
-	virtual TSharedRef<SWidget> GenerateWidgetForColumn( const FName& ColumnName ) OVERRIDE;
+	virtual TSharedRef<SWidget> GenerateWidgetForColumn( const FName& ColumnName ) override;
 
 	/** Override OnDragEnter for drag and drop of sockets onto bones */
-	virtual void OnDragEnter( const FGeometry& MyGeometry, const FDragDropEvent& DragDropEvent ) OVERRIDE;
+	virtual void OnDragEnter( const FGeometry& MyGeometry, const FDragDropEvent& DragDropEvent ) override;
 
 	/** Override OnDragLeave for drag and drop of sockets onto bones */
-	virtual void OnDragLeave( const FDragDropEvent& DragDropEvent ) OVERRIDE;
+	virtual void OnDragLeave( const FDragDropEvent& DragDropEvent ) override;
 
 	/** Override OnDrop for drag and drop of sockets and meshes onto bones */
-	virtual FReply OnDrop( const FGeometry& MyGeometry, const FDragDropEvent& DragDropEvent ) OVERRIDE;
+	virtual FReply OnDrop( const FGeometry& MyGeometry, const FDragDropEvent& DragDropEvent ) override;
 
 protected:
 
 	/** Handler for starting a drag/drop action */
-	virtual FReply OnDragDetected( const FGeometry& MyGeometry, const FPointerEvent& MouseEvent ) OVERRIDE;
+	virtual FReply OnDragDetected( const FGeometry& MyGeometry, const FPointerEvent& MouseEvent ) override;
 
 private:
 
diff --git a/Engine/Source/Editor/Persona/Private/SSkeletonTree.h b/Engine/Source/Editor/Persona/Private/SSkeletonTree.h
--- a/Engine/Source/Editor/Persona/Private/SSkeletonTree.h
+++ b/Engine/Source/Editor/Persona/Private/SSkeletonTree.h
@@ -148,7 +148,7 @@ public:
 	FReply OnDragDetected(const FGeometry& MyGeometry, const FPointerEvent& MouseEvent);
 
 	/** Return the name of the bone */
-	virtual FName GetRowItemName() const OVERRIDE {return BoneName;}
+	virtual FName GetRowItemName() const override {return BoneName;}
 
 	/** Set Translation Retargeting Mode for this bone. */
 	void SetBoneTranslationRetargetingMode(EBoneTranslationRetargetingMode::Type NewRetargetingMode);
@@ -227,7 +227,7 @@ public:
 	FReply OnDragDetected(const FGeometry& MyGeometry, const FPointerEvent& MouseEvent);
 
 	/** Return the name of the socket */
-	virtual FName GetRowItemName() const OVERRIDE {return SocketData->SocketName;}
+	virtual FName GetRowItemName() const override {return SocketData->SocketName;}
 
 	/** Handle double clicking a socket */
 	virtual void OnItemDoubleClicked();
@@ -240,7 +240,7 @@ public:
 	FOnRenameRequested OnRenameRequested;
 
 	/** Requests a rename on the socket item */
-	virtual void RequestRename() OVERRIDE;
+	virtual void RequestRename() override;
 
 	/** Return socket name as FText for display in skeleton tree */
 	FText GetSocketNameAsText() const { return FText::FromName(SocketData->SocketName); }
@@ -317,7 +317,7 @@ public:
 	virtual TSharedRef< SWidget > GenerateWidgetForDataColumn();
 
 	/** Return the name of the asset */
-	virtual FName GetRowItemName() const OVERRIDE {return FName( *Asset->GetName() ) ;}
+	virtual FName GetRowItemName() const override {return FName( *Asset->GetName() ) ;}
 
 	/** Return the name used to attach to this item (in assets case return the item we are attached to */
 	virtual FName GetAttachName() const { return GetParentName(); }
