Fixed crash in matinee when adding a new group to the track.
#TTP 337805 - Regression: EDITOR: MATINEE: CRASH: Adding an empty group to track results in a crash

#branch UE4

#proj Editor.Matinee

#change implemented the ActiveRecordCounts introduced in CL2092174, this means that it passes the check in CheckState.

diff --git a/Engine/Source/Editor/Matinee/Private/MatineeTransBuffer.cpp b/Engine/Source/Editor/Matinee/Private/MatineeTransBuffer.cpp
--- a/Engine/Source/Editor/Matinee/Private/MatineeTransBuffer.cpp
+++ b/Engine/Source/Editor/Matinee/Private/MatineeTransBuffer.cpp
@@ -15,6 +15,7 @@ UMatineeTransBuffer::UMatineeTransBuffer(const class FPostConstructInitializePro
 void UMatineeTransBuffer::BeginSpecial(const FText& Description)
 {
 	CheckState();
+	const int32 Result = ActiveCount;
 	if( ActiveCount++==0 )
 	{
 		// Cancel redo buffer.
@@ -34,6 +35,8 @@ void UMatineeTransBuffer::BeginSpecial(const FText& Description)
 		// Begin a new transaction.
 		GUndo = new(UndoBuffer)FMatineeTransaction( Description, 1 );
 	}
+	const int32 PriorRecordsCount = (Result > 0 ? ActiveRecordCounts[Result - 1] : 0);
+	ActiveRecordCounts.Add(UndoBuffer.Last().GetRecordCount() - PriorRecordsCount);
 	CheckState();
 }
 
@@ -45,5 +48,6 @@ void UMatineeTransBuffer::EndSpecial()
 	{
 		GUndo = NULL;
 	}
+	ActiveRecordCounts.Pop();
 	CheckState();
 }
