Fixed a crash caused by matinee transaction buffers not using the correct implementation of Begin

I've moved the implementation of UTransBuffer::Begin to a templated function that takes the transaction type to avoid issues like this arising in future.

Also corrected the return value of UTransBuffer::End which was returning the wrong value. This isn't used anywhere in the UE4 codebase but I've fixed it for correctness.
This addresses TTP#338312 - CRITICAL: Regression: EDITOR: MATINEE: CRASH: Occurs when clicking the Start/Stop record button

diff --git a/Engine/Source/Editor/Matinee/Private/MatineeTransBuffer.cpp b/Engine/Source/Editor/Matinee/Private/MatineeTransBuffer.cpp
--- a/Engine/Source/Editor/Matinee/Private/MatineeTransBuffer.cpp
+++ b/Engine/Source/Editor/Matinee/Private/MatineeTransBuffer.cpp
@@ -14,40 +14,10 @@ UMatineeTransBuffer::UMatineeTransBuffer(const class FPostConstructInitializePro
 
 void UMatineeTransBuffer::BeginSpecial(const FText& Description)
 {
-	CheckState();
-	const int32 Result = ActiveCount;
-	if( ActiveCount++==0 )
-	{
-		// Cancel redo buffer.
-		if( UndoCount )
-		{
-			UndoBuffer.RemoveAt( UndoBuffer.Num()-UndoCount, UndoCount );
-		}
-
-		UndoCount = 0;
-
-		// Purge previous transactions if too much data occupied.
-		while( GetUndoSize() > MaxMemory )
-		{
-			UndoBuffer.RemoveAt( 0 );
-		}
-
-		// Begin a new transaction.
-		GUndo = new(UndoBuffer)FMatineeTransaction( Description, 1 );
-	}
-	const int32 PriorRecordsCount = (Result > 0 ? ActiveRecordCounts[Result - 1] : 0);
-	ActiveRecordCounts.Add(UndoBuffer.Last().GetRecordCount() - PriorRecordsCount);
-	CheckState();
+	BeginInternal<FMatineeTransaction>(TEXT("Matinee"), Description);
 }
 
 void UMatineeTransBuffer::EndSpecial()
 {
-	CheckState();
-	check(ActiveCount>=1);
-	if( --ActiveCount==0 )
-	{
-		GUndo = NULL;
-	}
-	ActiveRecordCounts.Pop();
-	CheckState();
+	UTransBuffer::End();
 }
diff --git a/Engine/Source/Editor/Matinee/Private/MatineeTransaction.h b/Engine/Source/Editor/Matinee/Private/MatineeTransaction.h
--- a/Engine/Source/Editor/Matinee/Private/MatineeTransaction.h
+++ b/Engine/Source/Editor/Matinee/Private/MatineeTransaction.h
@@ -9,8 +9,8 @@ class FMatineeTransaction : public FTransaction
 {
 public:
 
-	FMatineeTransaction( const FText& InTitle=FText(), bool InFlip=0 )
-		:	FTransaction(TEXT("Matinee"), InTitle, InFlip)
+	FMatineeTransaction( const TCHAR* InContext = nullptr, const FText& InTitle=FText(), bool InFlip=0 )
+		:	FTransaction(InContext, InTitle, InFlip)
 	{ }
 
 public:
diff --git a/Engine/Source/Editor/UnrealEd/Classes/Editor/TransBuffer.h b/Engine/Source/Editor/UnrealEd/Classes/Editor/TransBuffer.h
--- a/Engine/Source/Editor/UnrealEd/Classes/Editor/TransBuffer.h
+++ b/Engine/Source/Editor/UnrealEd/Classes/Editor/TransBuffer.h
@@ -53,6 +53,38 @@ public:
 
 	// End UObject interface.
 
+protected:
+	
+	/** Implementation of the begin function. Used to create a specific transaction type */
+	template<typename TTransaction>
+	int32 BeginInternal( const TCHAR* SessionContext, const FText& Description )
+	{
+		CheckState();
+		const int32 Result = ActiveCount;
+		if (ActiveCount++ == 0)
+		{
+			// Cancel redo buffer.
+			if (UndoCount)
+			{
+				UndoBuffer.RemoveAt(UndoBuffer.Num() - UndoCount, UndoCount);
+			}
+			UndoCount = 0;
+
+			// Purge previous transactions if too much data occupied.
+			while (GetUndoSize() > MaxMemory)
+			{
+				UndoBuffer.RemoveAt(0);
+			}
+
+			// Begin a new transaction.
+			GUndo = new(UndoBuffer) TTransaction(SessionContext, Description, 1);
+		}
+		const int32 PriorRecordsCount = (Result > 0 ? ActiveRecordCounts[Result - 1] : 0);
+		ActiveRecordCounts.Add(UndoBuffer.Last().GetRecordCount() - PriorRecordsCount);
+		CheckState();
+		return Result;
+	}
+
 public:
 
 	// Begin UTransactor interface.
diff --git a/Engine/Source/Editor/UnrealEd/Private/EditorTransaction.cpp b/Engine/Source/Editor/UnrealEd/Private/EditorTransaction.cpp
--- a/Engine/Source/Editor/UnrealEd/Private/EditorTransaction.cpp
+++ b/Engine/Source/Editor/UnrealEd/Private/EditorTransaction.cpp
@@ -407,30 +407,7 @@ void UTransBuffer::AddReferencedObjects(UObject* InThis, FReferenceCollector& Co
 
 int32 UTransBuffer::Begin( const TCHAR* SessionContext, const FText& Description )
 {
-	CheckState();
-	const int32 Result = ActiveCount;
-	if (ActiveCount++ == 0)
-	{
-		// Cancel redo buffer.
-		if( UndoCount )
-		{
-			UndoBuffer.RemoveAt( UndoBuffer.Num()-UndoCount, UndoCount );
-		}
-		UndoCount = 0;
-
-		// Purge previous transactions if too much data occupied.
-		while( GetUndoSize() > MaxMemory )
-		{
-			UndoBuffer.RemoveAt( 0 );
-		}
-
-		// Begin a new transaction.
-		GUndo = new(UndoBuffer)FTransaction( SessionContext, Description, 1 );
-	}
-	const int32 PriorRecordsCount = (Result > 0 ? ActiveRecordCounts[Result - 1] : 0);
-	ActiveRecordCounts.Add(UndoBuffer.Last().GetRecordCount() - PriorRecordsCount);
-	CheckState();
-	return Result;
+	return BeginInternal<FTransaction>(SessionContext, Description);
 }
 
 
@@ -457,7 +434,7 @@ int32 UTransBuffer::End()
 		ActiveRecordCounts.Pop();
 		CheckState();
 	}
-	return ActiveCount;
+	return Result;
 }
 
 
