Added analytics reporting to matinee

#ttp 340617 - FIXIF: EDITOR:ANALYTICS: Matinee Instrumentation Pass

#branch UE4

#added Added analytics for various matinee usage; whever a track/group is created/removed, whenever anything is imported/exported and whenever anything is recorded

diff --git a/Engine/Source/Editor/Matinee/Private/Matinee.cpp b/Engine/Source/Editor/Matinee/Private/Matinee.cpp
--- a/Engine/Source/Editor/Matinee/Private/Matinee.cpp
+++ b/Engine/Source/Editor/Matinee/Private/Matinee.cpp
@@ -36,6 +36,8 @@
 #include "MessageLog.h"
 #include "Editor/PropertyEditor/Public/IDetailsView.h"
 
+#include "Runtime/Analytics/Analytics/Public/Interfaces/IAnalyticsProvider.h"
+#include "EngineAnalytics.h"
 
 DEFINE_LOG_CATEGORY(LogSlateMatinee);
 
@@ -1555,6 +1557,11 @@ void FMatinee::ToggleRecordInterpValues(void)
 	if (RecordingState != MatineeConstants::ERecordingState::RECORDING_COMPLETE)
 	{
 		StopRecordingInterpValues();
+
+		if (FEngineAnalytics::IsAvailable())
+		{
+			FEngineAnalytics::GetProvider().RecordEvent(TEXT("Editor.Matinee.Recorded"));
+		}
 	}
 	else
 	{
diff --git a/Engine/Source/Editor/Matinee/Private/MatineeMenus.cpp b/Engine/Source/Editor/Matinee/Private/MatineeMenus.cpp
--- a/Engine/Source/Editor/Matinee/Private/MatineeMenus.cpp
+++ b/Engine/Source/Editor/Matinee/Private/MatineeMenus.cpp
@@ -42,6 +42,8 @@
 #include "IDocumentation.h"
 #include "MainFrame.h"
 
+#include "Runtime/Analytics/Analytics/Public/Interfaces/IAnalyticsProvider.h"
+#include "EngineAnalytics.h"
 
 #define LOCTEXT_NAMESPACE "MatineeMenus"
 
@@ -370,10 +372,13 @@ void FMatinee::NewGroupPopupTextCommitted(
 	MatineeActor->Modify();
 	IData->Modify();
 
+	TArray< FAnalyticsEventAttribute > GroupAttribs;
+	GroupAttribs.Add(FAnalyticsEventAttribute(TEXT("ActionId"), FString::Printf(TEXT("%d"), static_cast<int32>(InActionId))));
 	if(bDirGroup)
 	{
 		UInterpGroup* NewDirector = ConstructObject<UInterpGroupDirector>( UInterpGroupDirector::StaticClass(), IData, NAME_None, RF_Transactional );
 		NewGroups.Add(NewDirector);
+		GroupAttribs.Add(FAnalyticsEventAttribute(TEXT("Name"), NewDirector->GroupName.ToString()));
 	}
 	else if(bDuplicateGroup)
 	{
@@ -389,6 +394,7 @@ void FMatinee::NewGroupPopupTextCommitted(
 			// which will crash again, so I'm disabling duplicating parenting. 
 			DupGroup->bIsParented = false;
 			NewGroups.Add(DupGroup);
+			GroupAttribs.Add(FAnalyticsEventAttribute(TEXT("Name"), DupGroup->GroupName.ToString()));
 		}
 	}
 	else
@@ -396,6 +402,12 @@ void FMatinee::NewGroupPopupTextCommitted(
 		UInterpGroup* NewGroup = ConstructObject<UInterpGroup>( UInterpGroup::StaticClass(), IData, NAME_None, RF_Transactional );
 		NewGroup->GroupName = NewGroupName;
 		NewGroups.Add(NewGroup);
+		GroupAttribs.Add(FAnalyticsEventAttribute(TEXT("Name"), NewGroup->GroupName.ToString()));
+	}
+
+	if (FEngineAnalytics::IsAvailable())
+	{
+		FEngineAnalytics::GetProvider().RecordEvent(TEXT("Editor.Matinee.NewGroup"), GroupAttribs);
 	}
 
 	IData->InterpGroups.Append(NewGroups);
@@ -3750,7 +3762,13 @@ void FMatinee::OnMenuImport()
 					}
 
 					// Re-create the Matinee sequence.
-					FFbxImporter->ImportMatineeSequence(MatineeActor);
+					if (FFbxImporter->ImportMatineeSequence(MatineeActor))
+					{
+						if (FEngineAnalytics::IsAvailable())
+						{
+							FEngineAnalytics::GetProvider().RecordEvent(TEXT("Editor.Matinee.Imported"));
+						}
+					}
 
 					// We have modified the sequence, so update its UI.
 					NotifyPostChange(NULL, NULL);
@@ -3819,8 +3837,14 @@ void FMatinee::OnMenuExport()
 					}
 				}
 
-			// Export Matinee
-			Exporter->ExportMatinee( MatineeActor );
+				// Export Matinee
+				if (Exporter->ExportMatinee(MatineeActor))
+				{
+					if (FEngineAnalytics::IsAvailable())
+					{
+						FEngineAnalytics::GetProvider().RecordEvent(TEXT("Editor.Matinee.Exported"));
+					}
+				}
 
 				// Save to disk
 				Exporter->WriteToFile( *ExportFilename );
@@ -3962,6 +3986,11 @@ void FMatinee::OnExportSoundCueInfoCommand()
 			// Close and delete archive
 			CSVFile->Close();
 			delete CSVFile;
+
+			if (FEngineAnalytics::IsAvailable())
+			{
+				FEngineAnalytics::GetProvider().RecordEvent(TEXT("Editor.Matinee.ExportedSoundCue"));
+			}
 		}
 		else
 		{
@@ -4146,6 +4175,11 @@ void FMatinee::OnExportAnimationInfoCommand()
 			// Close and delete archive
 			File->Close();
 			delete File;
+
+			if (FEngineAnalytics::IsAvailable())
+			{
+				FEngineAnalytics::GetProvider().RecordEvent(TEXT("Editor.Matinee.ExportedAnimationInfo"));
+			}
 		}
 	}
 }
diff --git a/Engine/Source/Editor/Matinee/Private/MatineeTools.cpp b/Engine/Source/Editor/Matinee/Private/MatineeTools.cpp
--- a/Engine/Source/Editor/Matinee/Private/MatineeTools.cpp
+++ b/Engine/Source/Editor/Matinee/Private/MatineeTools.cpp
@@ -34,6 +34,8 @@
 #include "SoundDefinitions.h"
 #include "EditorSupportDelegates.h"
 
+#include "Runtime/Analytics/Analytics/Public/Interfaces/IAnalyticsProvider.h"
+#include "EngineAnalytics.h"
 
 ///// UTILS
 
@@ -2776,6 +2778,11 @@ UInterpTrack* FMatinee::AddTrackToGroup( UInterpGroup* Group, UClass* TrackClass
 
 	check(NewTrack);
 
+	if (FEngineAnalytics::IsAvailable())
+	{
+		FEngineAnalytics::GetProvider().RecordEvent(TEXT("Editor.Matinee.NewTrack"), TEXT("Class"), TrackClass->GetName());
+	}
+
 	OutNewTrackIndex = Group->InterpTracks.Add(NewTrack);
 
 	check( NewTrack->TrackInstClass );
@@ -2945,9 +2952,15 @@ void FMatinee::DeleteSelectedTracks()
 		// Only allow base tracks to be deleted,  Subtracks will be deleted by their parent
 		if( ActiveTrack->GetOuter()->IsA( UInterpGroup::StaticClass() ) )
 		{
+			if (FEngineAnalytics::IsAvailable())
+			{
+				FEngineAnalytics::GetProvider().RecordEvent(TEXT("Editor.Matinee.DelTrack"), TEXT("Class"), ActiveTrack->GetClass()->GetName());
+			}
+
+			ActiveTrack->Modify();
+
 			UInterpGroup* Group = TrackIt.GetGroup();
 			Group->Modify();
-			ActiveTrack->Modify();
 
 			const int32 TrackIndex = TrackIt.GetTrackIndex();
 
@@ -3045,15 +3058,25 @@ void FMatinee::DeleteSelectedGroups()
 	for( FSelectedGroupIterator GroupIt(GetSelectedGroupIterator()); GroupIt; ++GroupIt )
 	{
 		UInterpGroup* GroupToDelete = *GroupIt;
+		if (FEngineAnalytics::IsAvailable())
+		{
+			FEngineAnalytics::GetProvider().RecordEvent(TEXT("Editor.Matinee.DelGroup"), TEXT("Name"), GroupToDelete->GroupName.ToString());
+		}
 
 		// Mark InterpGroup and all InterpTracks as Modified.
 		GroupToDelete->Modify();
 		for(int32 j=0; j<GroupToDelete->InterpTracks.Num(); j++)
 		{
-			GroupToDelete->InterpTracks[j]->Modify();
+			UInterpTrack* Track = GroupToDelete->InterpTracks[j];
+			if (FEngineAnalytics::IsAvailable())
+			{
+				FEngineAnalytics::GetProvider().RecordEvent(TEXT("Editor.Matinee.DelTrack"), TEXT("Class"), Track->GetClass()->GetName());
+			}
+
+			Track->Modify();
 
 			// Remove from the Curve editor, if its there.
-			IData->CurveEdSetup->RemoveCurve( GroupToDelete->InterpTracks[j] );
+			IData->CurveEdSetup->RemoveCurve(Track);
 		}
 		
 		// First, destroy any instances of this group.
diff --git a/Engine/Source/Editor/UnrealEd/Private/Fbx/FbxMainExport.cpp b/Engine/Source/Editor/UnrealEd/Private/Fbx/FbxMainExport.cpp
--- a/Engine/Source/Editor/UnrealEd/Private/Fbx/FbxMainExport.cpp
+++ b/Engine/Source/Editor/UnrealEd/Private/Fbx/FbxMainExport.cpp
@@ -954,9 +954,9 @@ FbxSurfaceMaterial* FFbxExporter::ExportMaterial(UMaterial* Material)
 /**
  * Exports the given Matinee sequence information into a FBX document.
  */
-void FFbxExporter::ExportMatinee(AMatineeActor* InMatineeActor)
+bool FFbxExporter::ExportMatinee(AMatineeActor* InMatineeActor)
 {
-	if (InMatineeActor == NULL || Scene == NULL) return;
+	if (InMatineeActor == NULL || Scene == NULL) return false;
 
 	// If the Matinee editor is not open, we need to initialize the sequence.
 	//bool InitializeMatinee = InMatineeActor->MatineeData == NULL;
@@ -1008,6 +1008,7 @@ void FFbxExporter::ExportMatinee(AMatineeActor* InMatineeActor)
 	//}
 
 	DefaultCamera = NULL;
+	return true;
 }
 
 
diff --git a/Engine/Source/Editor/UnrealEd/Private/Fbx/FbxMatineeImport.cpp b/Engine/Source/Editor/UnrealEd/Private/Fbx/FbxMatineeImport.cpp
--- a/Engine/Source/Editor/UnrealEd/Private/Fbx/FbxMatineeImport.cpp
+++ b/Engine/Source/Editor/UnrealEd/Private/Fbx/FbxMatineeImport.cpp
@@ -178,18 +178,18 @@ static FbxCamera* FindCamera( FbxNode* Parent )
 	return Camera;
 }
 
-void FFbxImporter::ImportMatineeSequence(AMatineeActor* InMatineeActor)
+bool FFbxImporter::ImportMatineeSequence(AMatineeActor* InMatineeActor)
 {
-	if (Scene == NULL || InMatineeActor == NULL) return;
+	if (Scene == NULL || InMatineeActor == NULL) return false;
 	
 	// merge animation layer at first
 	FbxAnimStack* AnimStack = Scene->GetMember<FbxAnimStack>(0);
-	if (!AnimStack) return;
+	if (!AnimStack) return false;
 		
 	MergeAllLayerAnimation(AnimStack, FbxTime::GetFrameRate(Scene->GetGlobalSettings().GetTimeMode()));
 
 	FbxAnimLayer* AnimLayer = AnimStack->GetMember<FbxAnimLayer>(0);
-	if (AnimLayer == NULL) return;
+	if (AnimLayer == NULL) return false;
 
 	// If the Matinee editor is not open, we need to initialize the sequence.
 	//bool InitializeMatinee = InMatineeActor->MatineeData == NULL;
@@ -287,6 +287,7 @@ void FFbxImporter::ImportMatineeSequence(AMatineeActor* InMatineeActor)
 	MatineeData->InterpLength = (InterpLength < 0.0f) ? 5.0f : InterpLength;
 	InMatineeActor->Modify();
 
+	return true;
 	//if (InitializeMatinee)
 	//{
 	//	InMatineeActor->TermInterp();
diff --git a/Engine/Source/Editor/UnrealEd/Private/FbxExporter.h b/Engine/Source/Editor/UnrealEd/Private/FbxExporter.h
--- a/Engine/Source/Editor/UnrealEd/Private/FbxExporter.h
+++ b/Engine/Source/Editor/UnrealEd/Private/FbxExporter.h
@@ -91,8 +91,10 @@ public:
 
 	/**
 	 * Exports the given Matinee sequence information into a FBX document.
+	 * 
+	 * @return	true, if sucessful
 	 */
-	virtual void ExportMatinee(class AMatineeActor* InMatineeActor);
+	virtual bool ExportMatinee(class AMatineeActor* InMatineeActor);
 
 	/**
 	 * Exports all the animation sequences part of a single Group in a Matinee sequence
diff --git a/Engine/Source/Editor/UnrealEd/Private/FbxImporter.h b/Engine/Source/Editor/UnrealEd/Private/FbxImporter.h
--- a/Engine/Source/Editor/UnrealEd/Private/FbxImporter.h
+++ b/Engine/Source/Editor/UnrealEd/Private/FbxImporter.h
@@ -900,8 +900,10 @@ public:
 	
 	/**
 	 * Modifies the Matinee actor with the animations found in the FBX document.
+	 * 
+	 * @return	true, if sucessful
 	 */
-	UNREALED_API void ImportMatineeSequence(AMatineeActor* InMatineeActor);
+	UNREALED_API bool ImportMatineeSequence(AMatineeActor* InMatineeActor);
 
 
 	/** Create a new asset from the package and objectname and class */
diff --git a/Engine/Source/Editor/UnrealEd/Public/MatineeExporter.h b/Engine/Source/Editor/UnrealEd/Public/MatineeExporter.h
--- a/Engine/Source/Editor/UnrealEd/Public/MatineeExporter.h
+++ b/Engine/Source/Editor/UnrealEd/Public/MatineeExporter.h
@@ -63,8 +63,10 @@ public:
 
 	/**
 	 * Exports the given Matinee sequence information into a file.
+	 * 
+	 * @return	true, if sucessful
 	 */
-	virtual void ExportMatinee(class AMatineeActor* InMatineeActor) = 0;
+	virtual bool ExportMatinee(class AMatineeActor* InMatineeActor) = 0;
 
 	/**
 	 * Writes the file to disk and releases it.
diff --git a/Engine/Source/Runtime/Engine/Private/Interpolation.cpp b/Engine/Source/Runtime/Engine/Private/Interpolation.cpp
--- a/Engine/Source/Runtime/Engine/Private/Interpolation.cpp
+++ b/Engine/Source/Runtime/Engine/Private/Interpolation.cpp
@@ -6840,6 +6840,8 @@ void UInterpTrackInstDirector::TermTrackInst(UInterpTrack* Track)
 	}
 	
 	OldViewTarget = NULL;
+
+	Super::TermTrackInst(Track);
 }
 
 
@@ -6930,6 +6932,8 @@ void UInterpTrackInstFade::TermTrackInst(UInterpTrack* Track)
 			}
 		}
 	}
+
+	Super::TermTrackInst(Track);
 }
 
 /*-----------------------------------------------------------------------------
@@ -7051,6 +7055,8 @@ void UInterpTrackInstSlomo::TermTrackInst(UInterpTrack* Track)
 		WorldSettings->MatineeTimeDilation = OldTimeDilation;
 		WorldSettings->ForceNetUpdate();
 	}
+
+	Super::TermTrackInst(Track);
 }
 
 /*-----------------------------------------------------------------------------
@@ -8102,6 +8108,8 @@ void UInterpTrackInstSound::TermTrackInst(UInterpTrack* Track)
 		}
 		PlayAudioComp = NULL;
 	}
+
+	Super::TermTrackInst(Track);
 }
 
 /*-----------------------------------------------------------------------------
@@ -8565,6 +8573,8 @@ void UInterpTrackInstFloatMaterialParam::TermTrackInst(UInterpTrack* Track)
 	MaterialInstances.Empty();
 	ResetFloats.Empty();
 	PrimitiveMaterialRefs.Empty();
+
+	Super::TermTrackInst(Track);
 }
 
 void UInterpTrackInstFloatMaterialParam::SaveActorState(UInterpTrack* Track)
@@ -8695,6 +8705,8 @@ void UInterpTrackInstVectorMaterialParam::TermTrackInst(UInterpTrack* Track)
 	MaterialInstances.Empty();
 	ResetVectors.Empty();
 	PrimitiveMaterialRefs.Empty();
+
+	Super::TermTrackInst(Track);
 }
 
 void UInterpTrackInstVectorMaterialParam::SaveActorState(UInterpTrack* Track)
@@ -8961,6 +8973,8 @@ void UInterpTrackInstColorScale::TermTrackInst(UInterpTrack* Track)
 		PC->PlayerCameraManager->bEnableColorScaling = false;
 		PC->PlayerCameraManager->ColorScale = FVector(1.f,1.f,1.f);
 	}
+
+	Super::TermTrackInst(Track);
 }
 
 /*-----------------------------------------------------------------------------
@@ -9050,6 +9064,7 @@ void UInterpTrackInstAudioMaster::InitTrackInst(UInterpTrack* Track)
 
 void UInterpTrackInstAudioMaster::TermTrackInst(UInterpTrack* Track)
 {
+	Super::TermTrackInst(Track);
 }
 
 
