Fix additional compile errors

diff --git a/Engine/Source/Editor/Matinee/Private/MatineeViewportClient.cpp b/Engine/Source/Editor/Matinee/Private/MatineeViewportClient.cpp
--- a/Engine/Source/Editor/Matinee/Private/MatineeViewportClient.cpp
+++ b/Engine/Source/Editor/Matinee/Private/MatineeViewportClient.cpp
@@ -165,6 +165,7 @@ bool FMatineeViewportClient::InputKey(FViewport* Viewport, int32 ControllerId, F
 	const bool bShiftDown = Viewport->KeyState(EKeys::LeftShift) || Viewport->KeyState(EKeys::RightShift);
 	const bool bAltDown = Viewport->KeyState(EKeys::LeftAlt) || Viewport->KeyState(EKeys::RightAlt);
 	const bool bCmdDown = Viewport->KeyState(EKeys::LeftCommand) || Viewport->KeyState(EKeys::RightCommand);
+	const bool bCapsDown = Viewport->KeyState(EKeys::CapsLock);
 
 	const int32 HitX = Viewport->GetMouseX();
 	const int32 HitY = Viewport->GetMouseY();
@@ -830,7 +831,7 @@ bool FMatineeViewportClient::InputKey(FViewport* Viewport, int32 ControllerId, F
 			InterpEd->ZoomView( 1.0f / FMatinee::InterpEditor_ZoomIncrement, InterpEd->bZoomToScrubPos );
 		}
 
-		FModifierKeysState ModKeys(bShiftDown, bShiftDown, bCtrlDown, bCtrlDown, bAltDown, bAltDown, bCmdDown, bCmdDown);
+		FModifierKeysState ModKeys(bShiftDown, bShiftDown, bCtrlDown, bCtrlDown, bAltDown, bAltDown, bCmdDown, bCmdDown, bCapsDown);
 		FKeyboardEvent KeyEvent(Key, ModKeys, false, 0);
 		InterpEd->ProcessCommandBindings(KeyEvent);
 	}
