Fixed matinee error suggesting the use of an InterpActor

diff --git a/Engine/Source/Editor/Matinee/Private/MatineeTools.cpp b/Engine/Source/Editor/Matinee/Private/MatineeTools.cpp
--- a/Engine/Source/Editor/Matinee/Private/MatineeTools.cpp
+++ b/Engine/Source/Editor/Matinee/Private/MatineeTools.cpp
@@ -2297,7 +2297,7 @@ void FMatinee::FinishAddKey(UInterpTrack* Track, bool bCommitKeys)
 	}
 	else
 	{
-		FMessageDialog::Open( EAppMsgType::Ok, NSLOCTEXT("UnrealEd", "NothingToKeyframe", "Nothing to keyframe, or selected object can't be keyframed on this type of track.\nFor example, you can't move a StaticMesh, you must instead use an InterpActor.") );
+		FMessageDialog::Open( EAppMsgType::Ok, NSLOCTEXT("UnrealEd", "NothingToKeyframe", "Nothing to keyframe, or selected object can't be keyframed on this type of track.\n") );
 	}
 
 	if (bCommitKeys)
