Fixed code relying on SLATE_TEXT_ARGUMENT for STextComboPopup

Made sure everything was using FText rather than FString.

diff --git a/Engine/Source/Editor/Matinee/Private/MatineeMenus.cpp b/Engine/Source/Editor/Matinee/Private/MatineeMenus.cpp
--- a/Engine/Source/Editor/Matinee/Private/MatineeMenus.cpp
+++ b/Engine/Source/Editor/Matinee/Private/MatineeMenus.cpp
@@ -2481,7 +2481,7 @@ void FMatinee::OnSetMoveKeyLookupGroup()
 
 		TSharedRef<STextComboPopup> TextEntryPopup = 
 			SNew(STextComboPopup)
-			.Label(NSLOCTEXT("Matinee.Popups", "SelectGroup", "Select Group").ToString())
+			.Label(NSLOCTEXT("Matinee.Popups", "SelectGroup", "Select Group"))
 			.TextOptions(GroupNames)
 			.OnTextChosen(this, &FMatinee::OnSetMoveKeyLookupGroupTextChosen, SelKey.KeyIndex, MoveTrack, MoveTrackAxis)
 			;
diff --git a/Engine/Source/Editor/Matinee/Private/MatineeTrackHelpers.cpp b/Engine/Source/Editor/Matinee/Private/MatineeTrackHelpers.cpp
--- a/Engine/Source/Editor/Matinee/Private/MatineeTrackHelpers.cpp
+++ b/Engine/Source/Editor/Matinee/Private/MatineeTrackHelpers.cpp
@@ -345,7 +345,7 @@ bool UMatineeTrackDirectorHelper::PreCreateKeyframe( UInterpTrack *Track, float
 
 		TSharedRef<STextComboPopup> TextEntryPopup =
 			SNew(STextComboPopup)
-			.Label(NSLOCTEXT("Matinee.Popups", "NewCut", "Cut to Group...").ToString())
+			.Label(NSLOCTEXT("Matinee.Popups", "NewCut", "Cut to Group..."))
 			.TextOptions(GroupNames)
 			.OnTextChosen_UObject(this, &UMatineeTrackDirectorHelper::OnAddKeyTextEntry, Mode->InterpEd, Track);
 
@@ -785,7 +785,7 @@ bool UMatineeTrackToggleHelper::PreCreateKeyframe( UInterpTrack *Track, float Ke
 
 	TSharedRef<STextComboPopup> TextEntryPopup = 
 		SNew(STextComboPopup)
-		.Label(NSLOCTEXT("Matinee.Popups", "ToggleAction", "Toggle Action").ToString())
+		.Label(NSLOCTEXT("Matinee.Popups", "ToggleAction", "Toggle Action"))
 		.TextOptions(PropStrings)
 		.OnTextChosen_UObject(this, &UMatineeTrackToggleHelper::OnAddKeyTextEntry, InterpEd, Track)
 		;
@@ -1138,7 +1138,7 @@ bool UMatineeTrackVisibilityHelper::PreCreateKeyframe( UInterpTrack *Track, floa
 
 	TSharedRef<STextComboPopup> TextEntryPopup = 
 		SNew(STextComboPopup)
-		.Label(NSLOCTEXT("Matinee.Popups", "VisibilityAction", "Visibility Action").ToString())
+		.Label(NSLOCTEXT("Matinee.Popups", "VisibilityAction", "Visibility Action"))
 		.TextOptions(PropStrings)
 		.OnTextChosen_UObject(this, &UMatineeTrackVisibilityHelper::OnAddKeyTextEntry, Mode->InterpEd, Track)
 		;
