fix for UE-1409 MATINEE: Undoing a rename of a camera group does not undo its name in a director group

diff --git a/Engine/Source/Editor/Matinee/Private/Matinee.cpp b/Engine/Source/Editor/Matinee/Private/Matinee.cpp
--- a/Engine/Source/Editor/Matinee/Private/Matinee.cpp
+++ b/Engine/Source/Editor/Matinee/Private/Matinee.cpp
@@ -498,6 +498,8 @@ void FMatinee::InitMatinee(const EToolkitMode::Type Mode, const TSharedPtr< clas
 
 	NormalTransactor = GEditor->Trans;
 	InterpEdTrans = new UMatineeTransBuffer( FObjectInitializer(), 8*1024*1024 );
+	InterpEdTrans->OnUndo().AddRaw(this, &FMatinee::OnPostUndoRedo);
+	InterpEdTrans->OnRedo().AddRaw(this, &FMatinee::OnPostUndoRedo);
 	GEditor->Trans = InterpEdTrans;
 
 	// Save viewports' data before it gets overridden by UpdateLevelViewport
@@ -1477,6 +1479,11 @@ void FMatinee::StartRecordingMovie()
 	GUnrealEd->PlayMap(NULL, NULL, 0, -1, false, true);
 }
 
+void FMatinee::OnPostUndoRedo(FUndoSessionContext SessionContext, bool Succeeded)
+{
+	InvalidateTrackWindowViewports();
+}
+
 //Key Command Helpers
 void FMatinee::OnMarkInSection()
 {
diff --git a/Engine/Source/Editor/Matinee/Private/Matinee.h b/Engine/Source/Editor/Matinee/Private/Matinee.h
--- a/Engine/Source/Editor/Matinee/Private/Matinee.h
+++ b/Engine/Source/Editor/Matinee/Private/Matinee.h
@@ -17,6 +17,7 @@
 #include "SMatineeRecorder.h"
 
 #include "Editor/DistCurveEditor/Public/IDistCurveEditor.h"
+#include "Editor/UnrealEd/Classes/Editor/Transactor.h"
 
 DECLARE_LOG_CATEGORY_EXTERN(LogSlateMatinee, Log, All);
 
@@ -99,6 +100,9 @@ public:
 	/** Starts recording the current sequence */
 	void StartRecordingMovie() override;
 
+	/** Handle undo redo events */
+	void OnPostUndoRedo(FUndoSessionContext SessionContext, bool Succeeded);
+
 	// Menu handlers
 	void OnMenuAddKey();
 	void OnMenuPlay(const bool bShouldLoop, const bool bPlayForward);
diff --git a/Engine/Source/Editor/Matinee/Private/MatineeMenus.cpp b/Engine/Source/Editor/Matinee/Private/MatineeMenus.cpp
--- a/Engine/Source/Editor/Matinee/Private/MatineeMenus.cpp
+++ b/Engine/Source/Editor/Matinee/Private/MatineeMenus.cpp
@@ -1739,6 +1739,7 @@ void FMatinee::OnContextGroupRenameCommitted(const FText& InText, ETextCommit::T
 		UInterpTrackDirector* DirTrack = DirGroup->GetDirectorTrack();
 		if(DirTrack)
 		{
+			DirTrack->Modify();
 			for(int32 i=0; i<DirTrack->CutTrack.Num(); i++)
 			{
 				FDirectorTrackCut& Cut = DirTrack->CutTrack[i];
