#Matinee: Export : Added config for exporting with hierarchy on

Merging

//depot/UE4-GDC2015/Engine/...

to //depot/UE4-GDC2015/Engine/...

diff --git a/Engine/Source/Editor/Matinee/Private/Matinee.cpp b/Engine/Source/Editor/Matinee/Private/Matinee.cpp
--- a/Engine/Source/Editor/Matinee/Private/Matinee.cpp
+++ b/Engine/Source/Editor/Matinee/Private/Matinee.cpp
@@ -430,7 +430,6 @@ void FMatinee::InitMatinee(const EToolkitMode::Type Mode, const TSharedPtr< clas
 	bEditingCrosshairEnabled = false;
 	bEditingGridEnabled = false;
 	bBakeTransforms = false;
-	bKeepHierarchy = false;
 	bAllowKeyframeBarSelection = false;
 	bAllowKeyframeTextSelection = false;
 	bLockCameraPitch = true;
diff --git a/Engine/Source/Editor/Matinee/Private/Matinee.h b/Engine/Source/Editor/Matinee/Private/Matinee.h
--- a/Engine/Source/Editor/Matinee/Private/Matinee.h
+++ b/Engine/Source/Editor/Matinee/Private/Matinee.h
@@ -1321,9 +1321,6 @@ protected:
 	/** When true, a key will be exported every frame instead of just the keys that user created. */
 	bool bBakeTransforms;
 
-	/** When true, the hierarchical relationship of objects will be kept when exported. */
-	bool bKeepHierarchy;
-
 	/** If true, clicking on a keyframe bar (such as one representing the duration of an audio cue, etc.) will cause a selection */
 	bool bAllowKeyframeBarSelection;
 
diff --git a/Engine/Source/Editor/Matinee/Private/MatineeMenus.cpp b/Engine/Source/Editor/Matinee/Private/MatineeMenus.cpp
--- a/Engine/Source/Editor/Matinee/Private/MatineeMenus.cpp
+++ b/Engine/Source/Editor/Matinee/Private/MatineeMenus.cpp
@@ -1580,6 +1580,7 @@ void FMatinee::OnContextTrackExportAnimFBX()
 					// Export the Matinee information to a COLLADA document.
 					Exporter->CreateDocument();
 					Exporter->SetTrasformBaking(bBakeTransforms);
+					const bool bKeepHierarchy = GEditor->AccessEditorUserSettings().bKeepAttachHierarchy;
 					Exporter->SetKeepHierarchy(bKeepHierarchy);
 
 					// Export the anim sequences
@@ -1971,6 +1972,7 @@ void FMatinee::OnContextGroupExportAnimFBX()
 							// Export the Matinee information to an FBX document.
 							Exporter->CreateDocument();
 							Exporter->SetTrasformBaking(bBakeTransforms);
+							const bool bKeepHierarchy = GEditor->AccessEditorUserSettings().bKeepAttachHierarchy;
 							Exporter->SetKeepHierarchy(bKeepHierarchy);
 
 							// Export the animation sequences in the group by sampling the skeletal mesh over the
@@ -3842,6 +3844,7 @@ void FMatinee::OnMenuExport()
 				// Export the Matinee information to an FBX file
 				Exporter->CreateDocument();
 				Exporter->SetTrasformBaking(bBakeTransforms);
+				const bool bKeepHierarchy = GEditor->AccessEditorUserSettings().bKeepAttachHierarchy;
 				Exporter->SetKeepHierarchy(bKeepHierarchy);
 
 				const bool bSelectedOnly = false;
@@ -4226,7 +4229,7 @@ bool FMatinee::IsBakeTransformsToggled()
  */
 void FMatinee::OnToggleKeepHierarchy()
 {
-	bKeepHierarchy = !bKeepHierarchy;
+	GEditor->AccessEditorUserSettings().bKeepAttachHierarchy = !GEditor->AccessEditorUserSettings().bKeepAttachHierarchy;
 }
 
 /**
@@ -4234,7 +4237,7 @@ void FMatinee::OnToggleKeepHierarchy()
  */
 bool FMatinee::IsKeepHierarchyToggled()
 {
-	return bKeepHierarchy;
+	return GEditor->AccessEditorUserSettings().bKeepAttachHierarchy;
 }
 
 
diff --git a/Engine/Source/Editor/UnrealEd/Classes/Editor/EditorUserSettings.h b/Engine/Source/Editor/UnrealEd/Classes/Editor/EditorUserSettings.h
--- a/Engine/Source/Editor/UnrealEd/Classes/Editor/EditorUserSettings.h
+++ b/Engine/Source/Editor/UnrealEd/Classes/Editor/EditorUserSettings.h
@@ -47,6 +47,10 @@ class UEditorUserSettings : public UObject
 	UPROPERTY(EditAnywhere, config, Category=HotReload)
 	uint32 bAutomaticallyHotReloadNewClasses:1;
 
+	/** If enabled, export level with attachment hierarchy set */
+	UPROPERTY(EditAnywhere, config, Category=Export)
+	uint32 bKeepAttachHierarchy:1;
+
 	/** Select to make Distributions use the curves, not the baked lookup tables. */
 	UPROPERTY(config)
 	uint32 bUseCurvesForDistributions:1; //(GDistributionType == 0)
