Unify Slate's FInputGesture and Engine's FInputChord as FInputChord defined in Slate

diff --git a/Engine/Source/Editor/Levels/Private/LevelCollectionViewCommands.h b/Engine/Source/Editor/Levels/Private/LevelCollectionViewCommands.h
--- a/Engine/Source/Editor/Levels/Private/LevelCollectionViewCommands.h
+++ b/Engine/Source/Editor/Levels/Private/LevelCollectionViewCommands.h
@@ -26,76 +26,76 @@ PRAGMA_DISABLE_OPTIMIZATION
 	virtual void RegisterCommands() override
 	{
 		//selected level
-		UI_COMMAND( MakeLevelCurrent, "Make Current", "Make this Level the Current Level", EUserInterfaceActionType::Button, FInputGesture( EKeys::Enter ) );
-		UI_COMMAND( MoveActorsToSelected, "Move Selected Actors to Level", "Moves the selected actors to this level", EUserInterfaceActionType::Button, FInputGesture() );
+		UI_COMMAND( MakeLevelCurrent, "Make Current", "Make this Level the Current Level", EUserInterfaceActionType::Button, FInputChord( EKeys::Enter ) );
+		UI_COMMAND( MoveActorsToSelected, "Move Selected Actors to Level", "Moves the selected actors to this level", EUserInterfaceActionType::Button, FInputChord() );
 
 		//invalid selected level
-		UI_COMMAND( FixUpInvalidReference, "Replace Missing Level","Removes the broken level and replaces it with the level browsed to", EUserInterfaceActionType::Button, FInputGesture() );
-		UI_COMMAND( RemoveInvalidReference, "Remove Missing Level", "Removes the reference to the missing level from the map", EUserInterfaceActionType::Button, FInputGesture() );
+		UI_COMMAND( FixUpInvalidReference, "Replace Missing Level","Removes the broken level and replaces it with the level browsed to", EUserInterfaceActionType::Button, FInputChord() );
+		UI_COMMAND( RemoveInvalidReference, "Remove Missing Level", "Removes the reference to the missing level from the map", EUserInterfaceActionType::Button, FInputChord() );
 
 		//levels
-		UI_COMMAND( EditProperties, "Edit Properties", "Opens the edit properties window for the selected levels", EUserInterfaceActionType::Button, FInputGesture() );
-		UI_COMMAND( SaveSelectedLevels, "Save Selected Levels", "Saves selected levels", EUserInterfaceActionType::Button, FInputGesture() );
-		UI_COMMAND( MigrateSelectedLevels, "Migrate Selected Levels...", "Copies the selected levels and all their dependencies to a different game", EUserInterfaceActionType::Button, FInputGesture() );
+		UI_COMMAND( EditProperties, "Edit Properties", "Opens the edit properties window for the selected levels", EUserInterfaceActionType::Button, FInputChord() );
+		UI_COMMAND( SaveSelectedLevels, "Save Selected Levels", "Saves selected levels", EUserInterfaceActionType::Button, FInputChord() );
+		UI_COMMAND( MigrateSelectedLevels, "Migrate Selected Levels...", "Copies the selected levels and all their dependencies to a different game", EUserInterfaceActionType::Button, FInputChord() );
 		
-		UI_COMMAND( DisplayActorCount, "Show Actor Count", "If enabled, displays the number of actors in each level", EUserInterfaceActionType::ToggleButton, FInputGesture() );
-		UI_COMMAND( DisplayLightmassSize, "Show Lightmass Size", "If enabled, displays the Lightmass size for each level", EUserInterfaceActionType::ToggleButton, FInputGesture() );
-		UI_COMMAND( DisplayFileSize, "Show File Size", "If enabled, displays the file size for each level", EUserInterfaceActionType::ToggleButton, FInputGesture() );
-		UI_COMMAND( DisplayPaths, "Show Paths", "If enabled, displays the path for each level", EUserInterfaceActionType::ToggleButton, FInputGesture() );
-		UI_COMMAND( DisplayEditorOffset, "Show Editor Offset", "If enabled, displays the editor offset for each level", EUserInterfaceActionType::ToggleButton, FInputGesture() );
+		UI_COMMAND( DisplayActorCount, "Show Actor Count", "If enabled, displays the number of actors in each level", EUserInterfaceActionType::ToggleButton, FInputChord() );
+		UI_COMMAND( DisplayLightmassSize, "Show Lightmass Size", "If enabled, displays the Lightmass size for each level", EUserInterfaceActionType::ToggleButton, FInputChord() );
+		UI_COMMAND( DisplayFileSize, "Show File Size", "If enabled, displays the file size for each level", EUserInterfaceActionType::ToggleButton, FInputChord() );
+		UI_COMMAND( DisplayPaths, "Show Paths", "If enabled, displays the path for each level", EUserInterfaceActionType::ToggleButton, FInputChord() );
+		UI_COMMAND( DisplayEditorOffset, "Show Editor Offset", "If enabled, displays the editor offset for each level", EUserInterfaceActionType::ToggleButton, FInputChord() );
 		
-		UI_COMMAND( CreateEmptyLevel, "Create Empty Level", "Creates a new empty Level", EUserInterfaceActionType::Button, FInputGesture() );
-		UI_COMMAND( AddExistingLevel, "Add Existing Level", "Adds an existing level", EUserInterfaceActionType::Button, FInputGesture() );
-		UI_COMMAND( AddSelectedActorsToNewLevel, "Add Selected Actors to New Level", "Adds the actors currently selected in the active viewport to a new Level", EUserInterfaceActionType::Button, FInputGesture() );
-		UI_COMMAND( MergeSelectedLevels, "Merge Levels into New Level", "Merges the selected levels into a new level, removing the original levels from the persistent", EUserInterfaceActionType::Button, FInputGesture() );
-		UI_COMMAND( RemoveSelectedLevels, "Remove Selected Levels", "Removes selected levels from world", EUserInterfaceActionType::Button, FInputGesture() );
+		UI_COMMAND( CreateEmptyLevel, "Create Empty Level", "Creates a new empty Level", EUserInterfaceActionType::Button, FInputChord() );
+		UI_COMMAND( AddExistingLevel, "Add Existing Level", "Adds an existing level", EUserInterfaceActionType::Button, FInputChord() );
+		UI_COMMAND( AddSelectedActorsToNewLevel, "Add Selected Actors to New Level", "Adds the actors currently selected in the active viewport to a new Level", EUserInterfaceActionType::Button, FInputChord() );
+		UI_COMMAND( MergeSelectedLevels, "Merge Levels into New Level", "Merges the selected levels into a new level, removing the original levels from the persistent", EUserInterfaceActionType::Button, FInputChord() );
+		UI_COMMAND( RemoveSelectedLevels, "Remove Selected Levels", "Removes selected levels from world", EUserInterfaceActionType::Button, FInputChord() );
 
 		// Source Control
-		UI_COMMAND( SCCCheckOut, "Check Out", "Checks out the selected asset from source control.", EUserInterfaceActionType::Button, FInputGesture() );
-		UI_COMMAND( SCCCheckIn, "Check In", "Checks in the selected asset to source control.", EUserInterfaceActionType::Button, FInputGesture() );
-		UI_COMMAND( SCCOpenForAdd, "Mark For Add", "Adds the selected asset to source control.", EUserInterfaceActionType::Button, FInputGesture() );
-		UI_COMMAND( SCCHistory, "History", "Displays the source control revision history of the selected asset.", EUserInterfaceActionType::Button, FInputGesture() );
-		UI_COMMAND( SCCRefresh, "Refresh", "Updates the source control status of the asset.", EUserInterfaceActionType::Button, FInputGesture() );
-		UI_COMMAND( SCCDiffAgainstDepot, "Diff Against Depot", "Look at differences between your version of the asset and that in source control.", EUserInterfaceActionType::Button, FInputGesture() );
-		UI_COMMAND( SCCConnect, "Connect To Source Control", "Connect to source control to allow source control operations to be performed on content and levels.", EUserInterfaceActionType::Button, FInputGesture() );
+		UI_COMMAND( SCCCheckOut, "Check Out", "Checks out the selected asset from source control.", EUserInterfaceActionType::Button, FInputChord() );
+		UI_COMMAND( SCCCheckIn, "Check In", "Checks in the selected asset to source control.", EUserInterfaceActionType::Button, FInputChord() );
+		UI_COMMAND( SCCOpenForAdd, "Mark For Add", "Adds the selected asset to source control.", EUserInterfaceActionType::Button, FInputChord() );
+		UI_COMMAND( SCCHistory, "History", "Displays the source control revision history of the selected asset.", EUserInterfaceActionType::Button, FInputChord() );
+		UI_COMMAND( SCCRefresh, "Refresh", "Updates the source control status of the asset.", EUserInterfaceActionType::Button, FInputChord() );
+		UI_COMMAND( SCCDiffAgainstDepot, "Diff Against Depot", "Look at differences between your version of the asset and that in source control.", EUserInterfaceActionType::Button, FInputChord() );
+		UI_COMMAND( SCCConnect, "Connect To Source Control", "Connect to source control to allow source control operations to be performed on content and levels.", EUserInterfaceActionType::Button, FInputChord() );
 
 		// set new streaming method
-		UI_COMMAND( SetAddStreamingMethod_Blueprint, "Set Blueprint Streaming Method", "Sets the streaming method for new or added levels to Blueprint streaming", EUserInterfaceActionType::RadioButton, FInputGesture() );
-		UI_COMMAND( SetAddStreamingMethod_AlwaysLoaded, "Set Streaming to Always Loaded", "Sets the streaming method new or added selected levels to be always loaded", EUserInterfaceActionType::RadioButton, FInputGesture() );
+		UI_COMMAND( SetAddStreamingMethod_Blueprint, "Set Blueprint Streaming Method", "Sets the streaming method for new or added levels to Blueprint streaming", EUserInterfaceActionType::RadioButton, FInputChord() );
+		UI_COMMAND( SetAddStreamingMethod_AlwaysLoaded, "Set Streaming to Always Loaded", "Sets the streaming method new or added selected levels to be always loaded", EUserInterfaceActionType::RadioButton, FInputChord() );
 
 		// change streaming method
-		UI_COMMAND( SetStreamingMethod_Blueprint, "Change Blueprint Streaming Method", "Changes the streaming method for the selected levels to Blueprint streaming", EUserInterfaceActionType::Check, FInputGesture() );
-		UI_COMMAND( SetStreamingMethod_AlwaysLoaded, "Change Streaming to Always Loaded", "Changes the streaming method for the selected levels to be always loaded", EUserInterfaceActionType::Check, FInputGesture() );
+		UI_COMMAND( SetStreamingMethod_Blueprint, "Change Blueprint Streaming Method", "Changes the streaming method for the selected levels to Blueprint streaming", EUserInterfaceActionType::Check, FInputChord() );
+		UI_COMMAND( SetStreamingMethod_AlwaysLoaded, "Change Streaming to Always Loaded", "Changes the streaming method for the selected levels to be always loaded", EUserInterfaceActionType::Check, FInputChord() );
 
 		//level selection
-		UI_COMMAND( SelectAllLevels, "Select All Levels", "Selects all levels", EUserInterfaceActionType::Button, FInputGesture() );
-		UI_COMMAND( DeselectAllLevels, "De-select All Levels", "De-selects all levels", EUserInterfaceActionType::Button, FInputGesture() );
-		UI_COMMAND( InvertLevelSelection, "Invert Level Selection", "Inverts level selection", EUserInterfaceActionType::Button, FInputGesture() );
+		UI_COMMAND( SelectAllLevels, "Select All Levels", "Selects all levels", EUserInterfaceActionType::Button, FInputChord() );
+		UI_COMMAND( DeselectAllLevels, "De-select All Levels", "De-selects all levels", EUserInterfaceActionType::Button, FInputChord() );
+		UI_COMMAND( InvertLevelSelection, "Invert Level Selection", "Inverts level selection", EUserInterfaceActionType::Button, FInputChord() );
 		
 		//actors
-		UI_COMMAND( SelectLevelActors, "Select Actors in Levels", "Sets actors in the selected Levels as the viewport's selection", EUserInterfaceActionType::Button, FInputGesture() );
-		UI_COMMAND( DeselectLevelActors, "Deselect Actors in Levels", "Removes the Actors in the selected Levels from the viewport's existing selection", EUserInterfaceActionType::Button, FInputGesture() );
+		UI_COMMAND( SelectLevelActors, "Select Actors in Levels", "Sets actors in the selected Levels as the viewport's selection", EUserInterfaceActionType::Button, FInputChord() );
+		UI_COMMAND( DeselectLevelActors, "Deselect Actors in Levels", "Removes the Actors in the selected Levels from the viewport's existing selection", EUserInterfaceActionType::Button, FInputChord() );
 
 		//streaming volumes
-		UI_COMMAND( AddStreamingLevelVolumes, "Add Streaming Volumes", "Adds the streaming volumes to the selected levels", EUserInterfaceActionType::Button, FInputGesture() );
-		UI_COMMAND( SetStreamingLevelVolumes, "Set Streaming Volumes", "Clears the streaming volumes associated with the selected levels and adds the selected streaming volumes", EUserInterfaceActionType::Button, FInputGesture() );
-		UI_COMMAND( SelectStreamingVolumes, "Select Streaming Volumes", "Selects the streaming volumes associated with the selected levels", EUserInterfaceActionType::Button, FInputGesture() );
-		UI_COMMAND( ClearStreamingVolumes, "Clear Streaming Volumes", "Removes all streaming volume associations with the selected levels", EUserInterfaceActionType::Button, FInputGesture() );
+		UI_COMMAND( AddStreamingLevelVolumes, "Add Streaming Volumes", "Adds the streaming volumes to the selected levels", EUserInterfaceActionType::Button, FInputChord() );
+		UI_COMMAND( SetStreamingLevelVolumes, "Set Streaming Volumes", "Clears the streaming volumes associated with the selected levels and adds the selected streaming volumes", EUserInterfaceActionType::Button, FInputChord() );
+		UI_COMMAND( SelectStreamingVolumes, "Select Streaming Volumes", "Selects the streaming volumes associated with the selected levels", EUserInterfaceActionType::Button, FInputChord() );
+		UI_COMMAND( ClearStreamingVolumes, "Clear Streaming Volumes", "Removes all streaming volume associations with the selected levels", EUserInterfaceActionType::Button, FInputChord() );
 
 		//visibility
-		UI_COMMAND( ShowSelectedLevels, "Show Selected Levels", "Toggles selected levels to a visible state in the viewports", EUserInterfaceActionType::Button, FInputGesture() );
-		UI_COMMAND( HideSelectedLevels, "Hide Selected Levels", "Toggles selected levels to an invisible state in the viewports", EUserInterfaceActionType::Button, FInputGesture() );
-		UI_COMMAND( ShowOnlySelectedLevels, "Show Only Selected Levels", "Toggles the selected levels to a visible state; toggles all other levels to an invisible state", EUserInterfaceActionType::Button, FInputGesture() );
-		UI_COMMAND( ShowAllLevels, "Show All Levels", "Toggles all levels to a visible state in the viewports", EUserInterfaceActionType::Button, FInputGesture() );
-		UI_COMMAND( HideAllLevels, "Hide All Levels", "Hides all levels to an invisible state in the viewports", EUserInterfaceActionType::Button, FInputGesture() );
+		UI_COMMAND( ShowSelectedLevels, "Show Selected Levels", "Toggles selected levels to a visible state in the viewports", EUserInterfaceActionType::Button, FInputChord() );
+		UI_COMMAND( HideSelectedLevels, "Hide Selected Levels", "Toggles selected levels to an invisible state in the viewports", EUserInterfaceActionType::Button, FInputChord() );
+		UI_COMMAND( ShowOnlySelectedLevels, "Show Only Selected Levels", "Toggles the selected levels to a visible state; toggles all other levels to an invisible state", EUserInterfaceActionType::Button, FInputChord() );
+		UI_COMMAND( ShowAllLevels, "Show All Levels", "Toggles all levels to a visible state in the viewports", EUserInterfaceActionType::Button, FInputChord() );
+		UI_COMMAND( HideAllLevels, "Hide All Levels", "Hides all levels to an invisible state in the viewports", EUserInterfaceActionType::Button, FInputChord() );
 
 		//lock
-		UI_COMMAND( LockSelectedLevels, "Lock Selected Levels", "Locks selected levels", EUserInterfaceActionType::Button, FInputGesture() );
-		UI_COMMAND( UnockSelectedLevels, "Unlock Selected Levels", "Unlocks selected levels", EUserInterfaceActionType::Button, FInputGesture() );
-		UI_COMMAND( LockAllLevels, "Lock All Levels", "Locks all levels", EUserInterfaceActionType::Button, FInputGesture() );
-		UI_COMMAND( UnockAllLevels, "Unlock All Levels", "Unlocks all levels", EUserInterfaceActionType::Button, FInputGesture() );
-		UI_COMMAND( LockReadOnlyLevels, "Lock Read-Only Levels", "Locks all read-only levels", EUserInterfaceActionType::Button, FInputGesture() );
-		UI_COMMAND( UnlockReadOnlyLevels, "Unlock Read-Only Levels", "Unlocks all read-only levels", EUserInterfaceActionType::Button, FInputGesture() );
+		UI_COMMAND( LockSelectedLevels, "Lock Selected Levels", "Locks selected levels", EUserInterfaceActionType::Button, FInputChord() );
+		UI_COMMAND( UnockSelectedLevels, "Unlock Selected Levels", "Unlocks selected levels", EUserInterfaceActionType::Button, FInputChord() );
+		UI_COMMAND( LockAllLevels, "Lock All Levels", "Locks all levels", EUserInterfaceActionType::Button, FInputChord() );
+		UI_COMMAND( UnockAllLevels, "Unlock All Levels", "Unlocks all levels", EUserInterfaceActionType::Button, FInputChord() );
+		UI_COMMAND( LockReadOnlyLevels, "Lock Read-Only Levels", "Locks all read-only levels", EUserInterfaceActionType::Button, FInputChord() );
+		UI_COMMAND( UnlockReadOnlyLevels, "Unlock Read-Only Levels", "Unlocks all read-only levels", EUserInterfaceActionType::Button, FInputChord() );
 	}
 PRAGMA_ENABLE_OPTIMIZATION
 
diff --git a/Engine/Source/Editor/LocalizationDashboard/Private/LocalizationTargetDetailCustomization.cpp b/Engine/Source/Editor/LocalizationDashboard/Private/LocalizationTargetDetailCustomization.cpp
--- a/Engine/Source/Editor/LocalizationDashboard/Private/LocalizationTargetDetailCustomization.cpp
+++ b/Engine/Source/Editor/LocalizationDashboard/Private/LocalizationTargetDetailCustomization.cpp
@@ -42,11 +42,11 @@ public:
 
 void FLocalizationTargetEditorCommands::RegisterCommands() 
 {
-	UI_COMMAND( Gather, "Gather", "Gathers text for this target.", EUserInterfaceActionType::Button, FInputGesture() );
-	UI_COMMAND( ImportAllCultures, "Import All", "Imports translations for all cultures of this target.", EUserInterfaceActionType::Button, FInputGesture() );
-	UI_COMMAND( ExportAllCultures, "Export All", "Exports translations for all cultures of this target.", EUserInterfaceActionType::Button, FInputGesture() );
-	UI_COMMAND( RefreshWordCounts, "Count Words", "Recalculates the word counts for all cultures of this target.", EUserInterfaceActionType::Button, FInputGesture() );
-	UI_COMMAND( Compile, "Compile", "Compiles translations for this target.", EUserInterfaceActionType::Button, FInputGesture() );
+	UI_COMMAND( Gather, "Gather", "Gathers text for this target.", EUserInterfaceActionType::Button, FInputChord() );
+	UI_COMMAND( ImportAllCultures, "Import All", "Imports translations for all cultures of this target.", EUserInterfaceActionType::Button, FInputChord() );
+	UI_COMMAND( ExportAllCultures, "Export All", "Exports translations for all cultures of this target.", EUserInterfaceActionType::Button, FInputChord() );
+	UI_COMMAND( RefreshWordCounts, "Count Words", "Recalculates the word counts for all cultures of this target.", EUserInterfaceActionType::Button, FInputChord() );
+	UI_COMMAND( Compile, "Compile", "Compiles translations for this target.", EUserInterfaceActionType::Button, FInputChord() );
 }
 
 void FLocalizationTargetDetailCustomization::CustomizeDetails(IDetailLayoutBuilder& DetailBuilder)
diff --git a/Engine/Source/Editor/LocalizationDashboard/Private/LocalizationTargetSetDetailCustomization.cpp b/Engine/Source/Editor/LocalizationDashboard/Private/LocalizationTargetSetDetailCustomization.cpp
--- a/Engine/Source/Editor/LocalizationDashboard/Private/LocalizationTargetSetDetailCustomization.cpp
+++ b/Engine/Source/Editor/LocalizationDashboard/Private/LocalizationTargetSetDetailCustomization.cpp
@@ -64,10 +64,10 @@ namespace
 
 	void FLocalizationDashboardCommands::RegisterCommands() 
 	{
-		UI_COMMAND( GatherAllTargets, "Gather All", "Gathers text for all targets in the project.", EUserInterfaceActionType::Button, FInputGesture() );
-		UI_COMMAND( ImportAllTargets, "Import All", "Imports translations for all cultures of all targets in the project.", EUserInterfaceActionType::Button, FInputGesture() );
-		UI_COMMAND( ExportAllTargets, "Export All", "Exports translations for all cultures of all targets in the project.", EUserInterfaceActionType::Button, FInputGesture() );
-		UI_COMMAND( CompileAllTargets, "Compile All", "Compiles translations for all targets in the project.", EUserInterfaceActionType::Button, FInputGesture() );
+		UI_COMMAND( GatherAllTargets, "Gather All", "Gathers text for all targets in the project.", EUserInterfaceActionType::Button, FInputChord() );
+		UI_COMMAND( ImportAllTargets, "Import All", "Imports translations for all cultures of all targets in the project.", EUserInterfaceActionType::Button, FInputChord() );
+		UI_COMMAND( ExportAllTargets, "Export All", "Exports translations for all cultures of all targets in the project.", EUserInterfaceActionType::Button, FInputChord() );
+		UI_COMMAND( CompileAllTargets, "Compile All", "Compiles translations for all targets in the project.", EUserInterfaceActionType::Button, FInputChord() );
 	}
 }
 
diff --git a/Engine/Source/Editor/MainFrame/Private/Frame/MainFrameActions.cpp b/Engine/Source/Editor/MainFrame/Private/Frame/MainFrameActions.cpp
--- a/Engine/Source/Editor/MainFrame/Private/Frame/MainFrameActions.cpp
+++ b/Engine/Source/Editor/MainFrame/Private/Frame/MainFrameActions.cpp
@@ -44,37 +44,37 @@ void FMainFrameCommands::RegisterCommands()
 	// Make a default can execute action that disables input when in debug mode
 	FCanExecuteAction DefaultExecuteAction = FCanExecuteAction::CreateStatic( &FMainFrameActionCallbacks::DefaultCanExecuteAction );
 
-	UI_COMMAND( SaveAll, "Save All", "Saves all unsaved levels and assets to disk", EUserInterfaceActionType::Button, FInputGesture( EModifierKey::Control, EKeys::S ) );
+	UI_COMMAND( SaveAll, "Save All", "Saves all unsaved levels and assets to disk", EUserInterfaceActionType::Button, FInputChord( EModifierKey::Control, EKeys::S ) );
 	ActionList->MapAction( SaveAll, FExecuteAction::CreateStatic( &FMainFrameActionCallbacks::SaveAll ), FCanExecuteAction::CreateStatic( &FMainFrameActionCallbacks::CanSaveWorld ) );
 
-	UI_COMMAND( ChooseFilesToSave, "Choose Files to Save...", "Opens a dialog with save options for content and levels", EUserInterfaceActionType::Button, FInputGesture() );
+	UI_COMMAND( ChooseFilesToSave, "Choose Files to Save...", "Opens a dialog with save options for content and levels", EUserInterfaceActionType::Button, FInputChord() );
 	ActionList->MapAction( ChooseFilesToSave, FExecuteAction::CreateStatic( &FMainFrameActionCallbacks::ChoosePackagesToSave ), FCanExecuteAction::CreateStatic( &FMainFrameActionCallbacks::CanSaveWorld ) );
 
-	UI_COMMAND( ChooseFilesToCheckIn, "Submit to Source Control...", "Opens a dialog with check in options for content and levels", EUserInterfaceActionType::Button, FInputGesture() );
+	UI_COMMAND( ChooseFilesToCheckIn, "Submit to Source Control...", "Opens a dialog with check in options for content and levels", EUserInterfaceActionType::Button, FInputChord() );
 	ActionList->MapAction( ChooseFilesToCheckIn, FExecuteAction::CreateStatic( &FMainFrameActionCallbacks::ChoosePackagesToCheckIn ), FCanExecuteAction::CreateStatic( &FMainFrameActionCallbacks::CanChoosePackagesToCheckIn ) );
 
-	UI_COMMAND( ConnectToSourceControl, "Connect To Source Control...", "Connect to source control to allow source control operations to be performed on content and levels.", EUserInterfaceActionType::Button, FInputGesture() );
+	UI_COMMAND( ConnectToSourceControl, "Connect To Source Control...", "Connect to source control to allow source control operations to be performed on content and levels.", EUserInterfaceActionType::Button, FInputChord() );
 	ActionList->MapAction( ConnectToSourceControl, FExecuteAction::CreateStatic( &FMainFrameActionCallbacks::ConnectToSourceControl ), DefaultExecuteAction );
 
-	UI_COMMAND( NewProject, "New Project...", "Opens a dialog to create a new game project", EUserInterfaceActionType::Button, FInputGesture() );
+	UI_COMMAND( NewProject, "New Project...", "Opens a dialog to create a new game project", EUserInterfaceActionType::Button, FInputChord() );
 	ActionList->MapAction( NewProject, FExecuteAction::CreateStatic( &FMainFrameActionCallbacks::NewProject, false, true), DefaultExecuteAction );
 
-	UI_COMMAND( OpenProject, "Open Project...", "Opens a dialog to choose a game project to open", EUserInterfaceActionType::Button, FInputGesture() );
+	UI_COMMAND( OpenProject, "Open Project...", "Opens a dialog to choose a game project to open", EUserInterfaceActionType::Button, FInputChord() );
 	ActionList->MapAction( OpenProject, FExecuteAction::CreateStatic( &FMainFrameActionCallbacks::NewProject, true, false), DefaultExecuteAction );
 
-	UI_COMMAND( AddCodeToProject, "New C++ Class...", "Adds C++ code to the project. The code can only be compiled if you have an appropriate C++ compiler installed.", EUserInterfaceActionType::Button, FInputGesture() );
+	UI_COMMAND( AddCodeToProject, "New C++ Class...", "Adds C++ code to the project. The code can only be compiled if you have an appropriate C++ compiler installed.", EUserInterfaceActionType::Button, FInputChord() );
 	ActionList->MapAction( AddCodeToProject, FExecuteAction::CreateStatic( &FMainFrameActionCallbacks::AddCodeToProject ));
 
-	UI_COMMAND( RefreshCodeProject, "Refresh code project", "Refreshes your C++ code project.", EUserInterfaceActionType::Button, FInputGesture() );
+	UI_COMMAND( RefreshCodeProject, "Refresh code project", "Refreshes your C++ code project.", EUserInterfaceActionType::Button, FInputChord() );
 	ActionList->MapAction( RefreshCodeProject, FExecuteAction::CreateStatic( &FMainFrameActionCallbacks::RefreshCodeProject ), DefaultExecuteAction );
 
-	UI_COMMAND( OpenIDE, "Open IDE", "Opens your C++ code in an integrated development environment.", EUserInterfaceActionType::Button, FInputGesture() );
+	UI_COMMAND( OpenIDE, "Open IDE", "Opens your C++ code in an integrated development environment.", EUserInterfaceActionType::Button, FInputChord() );
 	ActionList->MapAction( OpenIDE, FExecuteAction::CreateStatic( &FMainFrameActionCallbacks::OpenIDE ), DefaultExecuteAction );
 
-	UI_COMMAND( PackagingSettings, "Packaging Settings...", "Opens the settings for project packaging", EUserInterfaceActionType::Button, FInputGesture() );
+	UI_COMMAND( PackagingSettings, "Packaging Settings...", "Opens the settings for project packaging", EUserInterfaceActionType::Button, FInputChord() );
 	ActionList->MapAction( PackagingSettings, FExecuteAction::CreateStatic( &FMainFrameActionCallbacks::PackagingSettings ), DefaultExecuteAction );
 
-	//UI_COMMAND( LocalizeProject, "Localize Project...", "Opens the dashboard for managing project localization data.", EUserInterfaceActionType::Button, FInputGesture() );
+	//UI_COMMAND( LocalizeProject, "Localize Project...", "Opens the dashboard for managing project localization data.", EUserInterfaceActionType::Button, FInputChord() );
 	//ActionList->MapAction( LocalizeProject, FExecuteAction::CreateStatic( &FMainFrameActionCallbacks::LocalizeProject ), DefaultExecuteAction );
 
 	const int32 MaxProjects = 20;
@@ -92,7 +92,7 @@ void FMainFrameCommands::RegisterCommands()
 				Message,
 				LOCTEXT( "SwitchProjectToolTip", "Restarts the editor and switches to selected project" ) )
 			.UserInterfaceType( EUserInterfaceActionType::Button )
-			.DefaultGesture( FInputGesture() );
+			.DefaultChord( FInputChord() );
 		SwitchProjectCommands.Add( SwitchProject );
 
 		ActionList->MapAction( SwitchProjectCommands[ CurProjectIndex ], FExecuteAction::CreateStatic( &FMainFrameActionCallbacks::SwitchProjectByIndex, CurProjectIndex ),
@@ -100,63 +100,63 @@ void FMainFrameCommands::RegisterCommands()
 			FIsActionChecked::CreateStatic( &FMainFrameActionCallbacks::IsSwitchProjectChecked, CurProjectIndex ) );
 	}
 
-	UI_COMMAND( Exit, "Exit", "Exits the application", EUserInterfaceActionType::Button, FInputGesture() );
+	UI_COMMAND( Exit, "Exit", "Exits the application", EUserInterfaceActionType::Button, FInputChord() );
 	ActionList->MapAction( Exit, FExecuteAction::CreateStatic( &FMainFrameActionCallbacks::Exit ), DefaultExecuteAction );
 
 	ActionList->MapAction( FGenericCommands::Get().Undo, FExecuteAction::CreateStatic( &FMainFrameActionCallbacks::ExecuteExecCommand, FString( TEXT("TRANSACTION UNDO") ) ), FCanExecuteAction::CreateStatic( &FMainFrameActionCallbacks::Undo_CanExecute ) );
 
 	ActionList->MapAction( FGenericCommands::Get().Redo, FExecuteAction::CreateStatic( &FMainFrameActionCallbacks::ExecuteExecCommand, FString( TEXT("TRANSACTION REDO") ) ), FCanExecuteAction::CreateStatic( &FMainFrameActionCallbacks::Redo_CanExecute ) );
 
-	UI_COMMAND( OpenDeviceManagerApp, "Device Manager", "Opens up the device manager app", EUserInterfaceActionType::Check, FInputGesture() );
+	UI_COMMAND( OpenDeviceManagerApp, "Device Manager", "Opens up the device manager app", EUserInterfaceActionType::Check, FInputChord() );
 	ActionList->MapAction( OpenDeviceManagerApp, 
 												FExecuteAction::CreateStatic( &FMainFrameActionCallbacks::OpenSlateApp, FName( TEXT( "DeviceManager" ) ) ),
 												FCanExecuteAction(),
 												FIsActionChecked::CreateStatic( &FMainFrameActionCallbacks::OpenSlateApp_IsChecked, FName( TEXT( "DeviceManager" ) ) ) );
 
-	UI_COMMAND( OpenSessionManagerApp, "Session Manager", "Opens up the session manager app", EUserInterfaceActionType::Check, FInputGesture() );
+	UI_COMMAND( OpenSessionManagerApp, "Session Manager", "Opens up the session manager app", EUserInterfaceActionType::Check, FInputChord() );
 	ActionList->MapAction( OpenSessionManagerApp, 
 												FExecuteAction::CreateStatic( &FMainFrameActionCallbacks::OpenSlateApp, FName( "SessionFrontend" ) ),
 												FCanExecuteAction(),
 												FIsActionChecked::CreateStatic( &FMainFrameActionCallbacks::OpenSlateApp_IsChecked, FName("SessionFrontend" ) ) );
 
-	UI_COMMAND(VisitWiki, "Wiki...", "Go to the Unreal Engine Wiki page", EUserInterfaceActionType::Button, FInputGesture());
+	UI_COMMAND(VisitWiki, "Wiki...", "Go to the Unreal Engine Wiki page", EUserInterfaceActionType::Button, FInputChord());
 	ActionList->MapAction(VisitWiki, FExecuteAction::CreateStatic(&FMainFrameActionCallbacks::VisitWiki));
 
-	UI_COMMAND(VisitForums, "Forums...", "Go to the Unreal Engine forums", EUserInterfaceActionType::Button, FInputGesture());
+	UI_COMMAND(VisitForums, "Forums...", "Go to the Unreal Engine forums", EUserInterfaceActionType::Button, FInputChord());
 	ActionList->MapAction(VisitForums, FExecuteAction::CreateStatic(&FMainFrameActionCallbacks::VisitForums));
 
-	UI_COMMAND( VisitAskAQuestionPage, "Ask a Question...", "Have a question?  Go here to ask about anything and everything related to Unreal.", EUserInterfaceActionType::Button, FInputGesture() );
+	UI_COMMAND( VisitAskAQuestionPage, "Ask a Question...", "Have a question?  Go here to ask about anything and everything related to Unreal.", EUserInterfaceActionType::Button, FInputChord() );
 	ActionList->MapAction( VisitAskAQuestionPage, FExecuteAction::CreateStatic( &FMainFrameActionCallbacks::VisitAskAQuestionPage ) );
 
-	UI_COMMAND( VisitSearchForAnswersPage, "Answer Hub...", "Searches for useful answers on UDN provided by other users and experts.", EUserInterfaceActionType::Button, FInputGesture() );
+	UI_COMMAND( VisitSearchForAnswersPage, "Answer Hub...", "Searches for useful answers on UDN provided by other users and experts.", EUserInterfaceActionType::Button, FInputChord() );
 	ActionList->MapAction( VisitSearchForAnswersPage, FExecuteAction::CreateStatic( &FMainFrameActionCallbacks::VisitSearchForAnswersPage ) );
 
-	UI_COMMAND( VisitSupportWebSite, "Unreal Engine Support Web Site...", "Navigates to the Unreal Engine Support web site's main page.", EUserInterfaceActionType::Button, FInputGesture() );
+	UI_COMMAND( VisitSupportWebSite, "Unreal Engine Support Web Site...", "Navigates to the Unreal Engine Support web site's main page.", EUserInterfaceActionType::Button, FInputChord() );
 	ActionList->MapAction( VisitSupportWebSite, FExecuteAction::CreateStatic( &FMainFrameActionCallbacks::VisitSupportWebSite ) );
 
-	UI_COMMAND( VisitEpicGamesDotCom, "Visit UnrealEngine.com...", "Navigates to UnrealEngine.com where you can learn more about Unreal Technology.", EUserInterfaceActionType::Button, FInputGesture() );
+	UI_COMMAND( VisitEpicGamesDotCom, "Visit UnrealEngine.com...", "Navigates to UnrealEngine.com where you can learn more about Unreal Technology.", EUserInterfaceActionType::Button, FInputChord() );
 	ActionList->MapAction( VisitEpicGamesDotCom, FExecuteAction::CreateStatic( &FMainFrameActionCallbacks::VisitEpicGamesDotCom ) );
 
-	UI_COMMAND( AboutUnrealEd, "About Editor...", "Displays application credits and copyright information", EUserInterfaceActionType::Button, FInputGesture() );
+	UI_COMMAND( AboutUnrealEd, "About Editor...", "Displays application credits and copyright information", EUserInterfaceActionType::Button, FInputChord() );
 	ActionList->MapAction( AboutUnrealEd, FExecuteAction::CreateStatic( &FMainFrameActionCallbacks::AboutUnrealEd_Execute ) );
 
-	UI_COMMAND( CreditsUnrealEd, "Credits", "Displays application credits", EUserInterfaceActionType::Button, FInputGesture() );
+	UI_COMMAND( CreditsUnrealEd, "Credits", "Displays application credits", EUserInterfaceActionType::Button, FInputChord() );
 	ActionList->MapAction( CreditsUnrealEd, FExecuteAction::CreateStatic(&FMainFrameActionCallbacks::CreditsUnrealEd_Execute) );
 
-	UI_COMMAND( ResetLayout, "Reset Layout...", "Make a backup of your user settings and reset the layout customizations", EUserInterfaceActionType::Button, FInputGesture() );
+	UI_COMMAND( ResetLayout, "Reset Layout...", "Make a backup of your user settings and reset the layout customizations", EUserInterfaceActionType::Button, FInputChord() );
 	ActionList->MapAction( ResetLayout, FExecuteAction::CreateStatic( &FMainFrameActionCallbacks::ResetLayout) );
 
-	UI_COMMAND( SaveLayout, "Save Layout", "Save the layout customizations", EUserInterfaceActionType::Button, FInputGesture() );
+	UI_COMMAND( SaveLayout, "Save Layout", "Save the layout customizations", EUserInterfaceActionType::Button, FInputChord() );
 	ActionList->MapAction( SaveLayout, FExecuteAction::CreateStatic( &FMainFrameActionCallbacks::SaveLayout) );
 
-	UI_COMMAND( ToggleFullscreen, "Enable Fullscreen", "Enables fullscreen mode for the application, expanding across the entire monitor", EUserInterfaceActionType::ToggleButton, FInputGesture(EModifierKey::Shift, EKeys::F11) );
+	UI_COMMAND( ToggleFullscreen, "Enable Fullscreen", "Enables fullscreen mode for the application, expanding across the entire monitor", EUserInterfaceActionType::ToggleButton, FInputChord(EModifierKey::Shift, EKeys::F11) );
 	ActionList->MapAction( ToggleFullscreen,
 		FExecuteAction::CreateStatic( &FMainFrameActionCallbacks::ToggleFullscreen_Execute ),
 		FCanExecuteAction(),
 		FIsActionChecked::CreateStatic( &FMainFrameActionCallbacks::FullScreen_IsChecked )
 	);
 
-	UI_COMMAND(OpenWidgetReflector, "Open Widget Reflector", "Opens the Widget Reflector", EUserInterfaceActionType::Button, FInputGesture(EModifierKey::Shift | EModifierKey::Control , EKeys::W));
+	UI_COMMAND(OpenWidgetReflector, "Open Widget Reflector", "Opens the Widget Reflector", EUserInterfaceActionType::Button, FInputChord(EModifierKey::Shift | EModifierKey::Control , EKeys::W));
 	ActionList->MapAction(OpenWidgetReflector, FExecuteAction::CreateStatic(&FMainFrameActionCallbacks::OpenWidgetReflector_Execute));
 
 	FGlobalEditorCommonCommands::MapActions(ActionList);
diff --git a/Engine/Source/Editor/MaterialEditor/Private/MaterialEditor.cpp b/Engine/Source/Editor/MaterialEditor/Private/MaterialEditor.cpp
--- a/Engine/Source/Editor/MaterialEditor/Private/MaterialEditor.cpp
+++ b/Engine/Source/Editor/MaterialEditor/Private/MaterialEditor.cpp
@@ -4132,11 +4132,11 @@ void FMaterialEditor::OnNodeTitleCommitted(const FText& NewText, ETextCommit::Ty
 	}
 }
 
-FReply FMaterialEditor::OnSpawnGraphNodeByShortcut(FInputGesture InGesture, const FVector2D& InPosition, UEdGraph* InGraph)
+FReply FMaterialEditor::OnSpawnGraphNodeByShortcut(FInputChord InChord, const FVector2D& InPosition, UEdGraph* InGraph)
 {
 	UEdGraph* Graph = InGraph;
 
-	TSharedPtr< FEdGraphSchemaAction > Action = FMaterialEditorSpawnNodeCommands::Get().GetGraphActionByGesture(InGesture, InGraph);
+	TSharedPtr< FEdGraphSchemaAction > Action = FMaterialEditorSpawnNodeCommands::Get().GetGraphActionByChord(InChord, InGraph);
 
 	if(Action.IsValid())
 	{
diff --git a/Engine/Source/Editor/MaterialEditor/Private/MaterialEditor.h b/Engine/Source/Editor/MaterialEditor/Private/MaterialEditor.h
--- a/Engine/Source/Editor/MaterialEditor/Private/MaterialEditor.h
+++ b/Engine/Source/Editor/MaterialEditor/Private/MaterialEditor.h
@@ -386,15 +386,15 @@ protected:
 	void OnNodeTitleCommitted(const FText& NewText, ETextCommit::Type CommitInfo, UEdGraphNode* NodeBeingChanged);
 
 	/**
-	 * Handles spawning a graph node in the current graph using the passed in gesture
+	 * Handles spawning a graph node in the current graph using the passed in chord
 	 *
-	 * @param	InGesture	Gesture that was just performed
+	 * @param	InChord		Chord that was just performed
 	 * @param	InPosition	Current cursor position
-	 * @param	InGraph		Graph that gesture was performed in
+	 * @param	InGraph		Graph that chord was performed in
 	 *
-	 * @return	FReply	Whether gesture was handled
+	 * @return	FReply	Whether chord was handled
 	 */
-	FReply OnSpawnGraphNodeByShortcut(FInputGesture InGesture, const FVector2D& InPosition, UEdGraph* InGraph);
+	FReply OnSpawnGraphNodeByShortcut(FInputChord InChord, const FVector2D& InPosition, UEdGraph* InGraph);
 
 	/** Select every node in the graph */
 	void SelectAllNodes();
diff --git a/Engine/Source/Editor/MaterialEditor/Private/MaterialEditorActions.cpp b/Engine/Source/Editor/MaterialEditor/Private/MaterialEditorActions.cpp
--- a/Engine/Source/Editor/MaterialEditor/Private/MaterialEditorActions.cpp
+++ b/Engine/Source/Editor/MaterialEditor/Private/MaterialEditorActions.cpp
@@ -10,50 +10,50 @@
 
 void FMaterialEditorCommands::RegisterCommands()
 {
-	UI_COMMAND( Apply, "Apply", "Apply changes to original material and it's use in the world.", EUserInterfaceActionType::Button, FInputGesture(EKeys::Enter) );
-	UI_COMMAND( Flatten, "Flatten", "Flatten the material to a texture for mobile devices.", EUserInterfaceActionType::Button, FInputGesture() );
-
-	UI_COMMAND( ShowAllMaterialParameters, "Params", "Show or Hide all the materials parameters", EUserInterfaceActionType::ToggleButton, FInputGesture() );
-
-	UI_COMMAND( SetCylinderPreview, "Cylinder", "Sets the preview mesh to a cylinder primitive.", EUserInterfaceActionType::ToggleButton, FInputGesture() );
-	UI_COMMAND( SetSpherePreview, "Sphere", "Sets the preview mesh to a sphere primitive.", EUserInterfaceActionType::ToggleButton, FInputGesture() );
-	UI_COMMAND( SetPlanePreview, "Plane", "Sets the preview mesh to a plane primitive.", EUserInterfaceActionType::ToggleButton, FInputGesture() );
-	UI_COMMAND( SetCubePreview, "Cube", "Sets the preview mesh to a cube primitive.", EUserInterfaceActionType::ToggleButton, FInputGesture() );
-	UI_COMMAND( SetPreviewMeshFromSelection, "Mesh", "Sets the preview mesh based on the current content browser selection.", EUserInterfaceActionType::ToggleButton, FInputGesture() );
-	UI_COMMAND( TogglePreviewGrid, "Grid", "Toggles the preview pane's grid.", EUserInterfaceActionType::ToggleButton, FInputGesture() );
-	UI_COMMAND( TogglePreviewBackground, "Background", "Toggles the preview pane's background.", EUserInterfaceActionType::ToggleButton, FInputGesture() );
-
-	UI_COMMAND( CameraHome, "Home", "Goes home on the canvas.", EUserInterfaceActionType::Button, FInputGesture() );
-	UI_COMMAND( CleanUnusedExpressions, "Clean Up", "Cleans up any unused Expressions.", EUserInterfaceActionType::Button, FInputGesture() );
-	UI_COMMAND(ShowHideConnectors, "Connectors", "Show or Hide Unused Connectors", EUserInterfaceActionType::ToggleButton, FInputGesture());
-	UI_COMMAND( ToggleLivePreview, "Live Preview", "Toggles real time update of the preview material.", EUserInterfaceActionType::ToggleButton, FInputGesture());
-	UI_COMMAND( ToggleRealtimeExpressions, "Live Nodes", "Toggles real time update of the graph canvas.", EUserInterfaceActionType::ToggleButton, FInputGesture() );
-	UI_COMMAND( AlwaysRefreshAllPreviews, "Live Update", "All nodes are previewed live.", EUserInterfaceActionType::ToggleButton, FInputGesture() );
-	UI_COMMAND( ToggleMaterialStats, "Stats", "Toggles displaying of the material's stats.", EUserInterfaceActionType::ToggleButton, FInputGesture());
-	UI_COMMAND( ToggleMobileStats, "Mobile Stats", "Toggles material stats and compilation errors for mobile.", EUserInterfaceActionType::ToggleButton, FInputGesture() );
-	UI_COMMAND( NewComment, "New Comment", "Creates a new comment node.", EUserInterfaceActionType::Button, FInputGesture() );
-	UI_COMMAND( MatertialPasteHere, "Paste Here", "Pastes copied items at this location.", EUserInterfaceActionType::Button, FInputGesture() );
-	UI_COMMAND( UseCurrentTexture, "Use Current Texture", "Uses the current texture selected in the content browser.", EUserInterfaceActionType::Button, FInputGesture() );
-	UI_COMMAND( ConvertObjects, "Convert to Parameter", "Converts the objects to parameters.", EUserInterfaceActionType::Button, FInputGesture() );
-	UI_COMMAND( ConvertToConstant, "Convert to Constant", "Converts the parameters to constants.", EUserInterfaceActionType::Button, FInputGesture() );
-	UI_COMMAND( ConvertToTextureObjects, "Convert to Texture Object", "Converts the objects to texture objects.", EUserInterfaceActionType::Button, FInputGesture() );
-	UI_COMMAND( ConvertToTextureSamples, "Convert to Texture Sample", "Converts the objects to texture samples.", EUserInterfaceActionType::Button, FInputGesture() );
-	UI_COMMAND( StopPreviewNode, "Stop Previewing Node", "Stops the preview viewport from previewing this node", EUserInterfaceActionType::Button, FInputGesture() );
-	UI_COMMAND( StartPreviewNode, "Start Previewing Node", "Makes the preview viewport start previewing this node", EUserInterfaceActionType::Button, FInputGesture() );
-	UI_COMMAND( EnableRealtimePreviewNode, "Enable Realtime Preview", "Enables realtime previewing of this expression node", EUserInterfaceActionType::Button, FInputGesture() );
-	UI_COMMAND( DisableRealtimePreviewNode, "Disable Realtime Preview", "Disables realtime previewing of this expression node", EUserInterfaceActionType::Button, FInputGesture() );
-	UI_COMMAND( BreakAllLinks, "Break All Links", "Breaks all links leading out of this node.", EUserInterfaceActionType::Button, FInputGesture() );
-	UI_COMMAND( DuplicateObjects, "Duplicate Object(s)", "Duplicates the selected objects.", EUserInterfaceActionType::Button, FInputGesture() );
-	UI_COMMAND( DeleteObjects, "Delete Object(s)", "Deletes the selected objects.", EUserInterfaceActionType::Button, FInputGesture() );
-	UI_COMMAND( SelectDownstreamNodes, "Select Downstream Nodes", "Selects all nodes that use this node's outgoing links.", EUserInterfaceActionType::Button, FInputGesture() );
-	UI_COMMAND( SelectUpstreamNodes, "Select Upstream Nodes", "Selects all nodes that feed links into this node.", EUserInterfaceActionType::Button, FInputGesture() );
-	UI_COMMAND( RemoveFromFavorites, "Remove From Favorites", "Removes this expression from your favorites.", EUserInterfaceActionType::Button, FInputGesture() );
-	UI_COMMAND( AddToFavorites, "Add To Favorites", "Adds this expression to your favorites.", EUserInterfaceActionType::Button, FInputGesture() );
-	UI_COMMAND( BreakLink, "Break Link", "Deletes this link.", EUserInterfaceActionType::Button, FInputGesture() );
-	UI_COMMAND( ForceRefreshPreviews, "Force Refresh Previews", "Forces a refresh of all previews", EUserInterfaceActionType::Button, FInputGesture(EKeys::SpaceBar) );
-	UI_COMMAND( CreateComponentMaskNode, "Create ComponentMask Node", "Creates a ComponentMask node at the current cursor position.", EUserInterfaceActionType::Button, FInputGesture(EModifierKey::Shift, EKeys::C));
-	UI_COMMAND( GoToDocumentation, "View Documentation", "View documentation for this node.", EUserInterfaceActionType::Button, FInputGesture());
-	UI_COMMAND( FindInMaterial, "Search", "Finds expressions and comments in the current Material", EUserInterfaceActionType::Button, FInputGesture(EModifierKey::Control, EKeys::F));
+	UI_COMMAND( Apply, "Apply", "Apply changes to original material and it's use in the world.", EUserInterfaceActionType::Button, FInputChord(EKeys::Enter) );
+	UI_COMMAND( Flatten, "Flatten", "Flatten the material to a texture for mobile devices.", EUserInterfaceActionType::Button, FInputChord() );
+
+	UI_COMMAND( ShowAllMaterialParameters, "Params", "Show or Hide all the materials parameters", EUserInterfaceActionType::ToggleButton, FInputChord() );
+
+	UI_COMMAND( SetCylinderPreview, "Cylinder", "Sets the preview mesh to a cylinder primitive.", EUserInterfaceActionType::ToggleButton, FInputChord() );
+	UI_COMMAND( SetSpherePreview, "Sphere", "Sets the preview mesh to a sphere primitive.", EUserInterfaceActionType::ToggleButton, FInputChord() );
+	UI_COMMAND( SetPlanePreview, "Plane", "Sets the preview mesh to a plane primitive.", EUserInterfaceActionType::ToggleButton, FInputChord() );
+	UI_COMMAND( SetCubePreview, "Cube", "Sets the preview mesh to a cube primitive.", EUserInterfaceActionType::ToggleButton, FInputChord() );
+	UI_COMMAND( SetPreviewMeshFromSelection, "Mesh", "Sets the preview mesh based on the current content browser selection.", EUserInterfaceActionType::ToggleButton, FInputChord() );
+	UI_COMMAND( TogglePreviewGrid, "Grid", "Toggles the preview pane's grid.", EUserInterfaceActionType::ToggleButton, FInputChord() );
+	UI_COMMAND( TogglePreviewBackground, "Background", "Toggles the preview pane's background.", EUserInterfaceActionType::ToggleButton, FInputChord() );
+
+	UI_COMMAND( CameraHome, "Home", "Goes home on the canvas.", EUserInterfaceActionType::Button, FInputChord() );
+	UI_COMMAND( CleanUnusedExpressions, "Clean Up", "Cleans up any unused Expressions.", EUserInterfaceActionType::Button, FInputChord() );
+	UI_COMMAND(ShowHideConnectors, "Connectors", "Show or Hide Unused Connectors", EUserInterfaceActionType::ToggleButton, FInputChord());
+	UI_COMMAND( ToggleLivePreview, "Live Preview", "Toggles real time update of the preview material.", EUserInterfaceActionType::ToggleButton, FInputChord());
+	UI_COMMAND( ToggleRealtimeExpressions, "Live Nodes", "Toggles real time update of the graph canvas.", EUserInterfaceActionType::ToggleButton, FInputChord() );
+	UI_COMMAND( AlwaysRefreshAllPreviews, "Live Update", "All nodes are previewed live.", EUserInterfaceActionType::ToggleButton, FInputChord() );
+	UI_COMMAND( ToggleMaterialStats, "Stats", "Toggles displaying of the material's stats.", EUserInterfaceActionType::ToggleButton, FInputChord());
+	UI_COMMAND( ToggleMobileStats, "Mobile Stats", "Toggles material stats and compilation errors for mobile.", EUserInterfaceActionType::ToggleButton, FInputChord() );
+	UI_COMMAND( NewComment, "New Comment", "Creates a new comment node.", EUserInterfaceActionType::Button, FInputChord() );
+	UI_COMMAND( MatertialPasteHere, "Paste Here", "Pastes copied items at this location.", EUserInterfaceActionType::Button, FInputChord() );
+	UI_COMMAND( UseCurrentTexture, "Use Current Texture", "Uses the current texture selected in the content browser.", EUserInterfaceActionType::Button, FInputChord() );
+	UI_COMMAND( ConvertObjects, "Convert to Parameter", "Converts the objects to parameters.", EUserInterfaceActionType::Button, FInputChord() );
+	UI_COMMAND( ConvertToConstant, "Convert to Constant", "Converts the parameters to constants.", EUserInterfaceActionType::Button, FInputChord() );
+	UI_COMMAND( ConvertToTextureObjects, "Convert to Texture Object", "Converts the objects to texture objects.", EUserInterfaceActionType::Button, FInputChord() );
+	UI_COMMAND( ConvertToTextureSamples, "Convert to Texture Sample", "Converts the objects to texture samples.", EUserInterfaceActionType::Button, FInputChord() );
+	UI_COMMAND( StopPreviewNode, "Stop Previewing Node", "Stops the preview viewport from previewing this node", EUserInterfaceActionType::Button, FInputChord() );
+	UI_COMMAND( StartPreviewNode, "Start Previewing Node", "Makes the preview viewport start previewing this node", EUserInterfaceActionType::Button, FInputChord() );
+	UI_COMMAND( EnableRealtimePreviewNode, "Enable Realtime Preview", "Enables realtime previewing of this expression node", EUserInterfaceActionType::Button, FInputChord() );
+	UI_COMMAND( DisableRealtimePreviewNode, "Disable Realtime Preview", "Disables realtime previewing of this expression node", EUserInterfaceActionType::Button, FInputChord() );
+	UI_COMMAND( BreakAllLinks, "Break All Links", "Breaks all links leading out of this node.", EUserInterfaceActionType::Button, FInputChord() );
+	UI_COMMAND( DuplicateObjects, "Duplicate Object(s)", "Duplicates the selected objects.", EUserInterfaceActionType::Button, FInputChord() );
+	UI_COMMAND( DeleteObjects, "Delete Object(s)", "Deletes the selected objects.", EUserInterfaceActionType::Button, FInputChord() );
+	UI_COMMAND( SelectDownstreamNodes, "Select Downstream Nodes", "Selects all nodes that use this node's outgoing links.", EUserInterfaceActionType::Button, FInputChord() );
+	UI_COMMAND( SelectUpstreamNodes, "Select Upstream Nodes", "Selects all nodes that feed links into this node.", EUserInterfaceActionType::Button, FInputChord() );
+	UI_COMMAND( RemoveFromFavorites, "Remove From Favorites", "Removes this expression from your favorites.", EUserInterfaceActionType::Button, FInputChord() );
+	UI_COMMAND( AddToFavorites, "Add To Favorites", "Adds this expression to your favorites.", EUserInterfaceActionType::Button, FInputChord() );
+	UI_COMMAND( BreakLink, "Break Link", "Deletes this link.", EUserInterfaceActionType::Button, FInputChord() );
+	UI_COMMAND( ForceRefreshPreviews, "Force Refresh Previews", "Forces a refresh of all previews", EUserInterfaceActionType::Button, FInputChord(EKeys::SpaceBar) );
+	UI_COMMAND( CreateComponentMaskNode, "Create ComponentMask Node", "Creates a ComponentMask node at the current cursor position.", EUserInterfaceActionType::Button, FInputChord(EModifierKey::Shift, EKeys::C));
+	UI_COMMAND( GoToDocumentation, "View Documentation", "View documentation for this node.", EUserInterfaceActionType::Button, FInputChord());
+	UI_COMMAND( FindInMaterial, "Search", "Finds expressions and comments in the current Material", EUserInterfaceActionType::Button, FInputChord(EModifierKey::Control, EKeys::F));
 }
 
 //////////////////////////////////////////////////////////////////////////
@@ -128,12 +128,12 @@ void FMaterialEditorSpawnNodeCommands::RegisterCommands()
 				FParse::Bool(*NodeSpawns[x], TEXT("Ctrl="), bCtrl);
 			}
 
-			FInputGesture Gesture(Key, EModifierKey::FromBools(bCtrl, bAlt, bShift, false));
+			FInputChord Chord(Key, EModifierKey::FromBools(bCtrl, bAlt, bShift, false));
 
 			const FText CommandLabelText = FText::FromString( CommandLabel );
 			const FText Description = FText::Format( NSLOCTEXT("MaterialEditor", "NodeSpawnDescription", "Hold down the bound keys and left click in the graph panel to spawn a {0} node."), CommandLabelText);
 
-			FUICommandInfo::MakeCommandInfo( this->AsShared(), CommandInfo, FName(*NodeSpawns[x]), CommandLabelText, Description, FSlateIcon(FEditorStyle::GetStyleSetName(), *FString::Printf(TEXT("%s.%s"), *this->GetContextName().ToString(), *NodeSpawns[x])), EUserInterfaceActionType::Button, Gesture );
+			FUICommandInfo::MakeCommandInfo( this->AsShared(), CommandInfo, FName(*NodeSpawns[x]), CommandLabelText, Description, FSlateIcon(FEditorStyle::GetStyleSetName(), *FString::Printf(TEXT("%s.%s"), *this->GetContextName().ToString(), *NodeSpawns[x])), EUserInterfaceActionType::Button, Chord );
 
 			InfoPtr->CommandInfo = CommandInfo;
 
@@ -142,13 +142,13 @@ void FMaterialEditorSpawnNodeCommands::RegisterCommands()
 	}
 }
 
-TSharedPtr< FEdGraphSchemaAction > FMaterialEditorSpawnNodeCommands::GetGraphActionByGesture(FInputGesture& InGesture, UEdGraph* InDestGraph) const
+TSharedPtr< FEdGraphSchemaAction > FMaterialEditorSpawnNodeCommands::GetGraphActionByChord(FInputChord& InChord, UEdGraph* InDestGraph) const
 {
-	if(InGesture.IsValidGesture())
+	if(InChord.IsValidChord())
 	{
 		for(int32 x = 0; x < NodeCommands.Num(); ++x)
 		{
-			if(NodeCommands[x]->CommandInfo->GetActiveGesture().Get() == InGesture)
+			if(NodeCommands[x]->CommandInfo->GetActiveChord().Get() == InChord)
 			{
 				return NodeCommands[x]->GetAction(InDestGraph);
 			}
@@ -158,17 +158,17 @@ TSharedPtr< FEdGraphSchemaAction > FMaterialEditorSpawnNodeCommands::GetGraphAct
 	return TSharedPtr< FEdGraphSchemaAction >();
 }
 
-const TSharedPtr<const FInputGesture> FMaterialEditorSpawnNodeCommands::GetGestureByClass(UClass* MaterialExpressionClass) const
+const TSharedPtr<const FInputChord> FMaterialEditorSpawnNodeCommands::GetChordByClass(UClass* MaterialExpressionClass) const
 {
 	for(int32 Index = 0; Index < NodeCommands.Num(); ++Index)
 	{
 		if (NodeCommands[Index]->GetClass() == MaterialExpressionClass)
 		{
-			return NodeCommands[Index]->CommandInfo->GetActiveGesture();
+			return NodeCommands[Index]->CommandInfo->GetActiveChord();
 		}
 	}
 
-	return TSharedPtr< const FInputGesture >();
+	return TSharedPtr< const FInputChord >();
 }
 
 #undef LOCTEXT_NAMESPACE
\ No newline at end of file
diff --git a/Engine/Source/Editor/MaterialEditor/Private/SMaterialPalette.cpp b/Engine/Source/Editor/MaterialEditor/Private/SMaterialPalette.cpp
--- a/Engine/Source/Editor/MaterialEditor/Private/SMaterialPalette.cpp
+++ b/Engine/Source/Editor/MaterialEditor/Private/SMaterialPalette.cpp
@@ -23,16 +23,16 @@ void SMaterialPaletteItem::Construct(const FArguments& InArgs, FCreateWidgetForA
 	TSharedPtr<FEdGraphSchemaAction> GraphAction = InCreateData->Action;
 	ActionPtr = InCreateData->Action;
 
-	// Get the Hotkey gesture if one exists for this action
-	TSharedPtr<const FInputGesture> HotkeyGesture;
+	// Get the Hotkey chord if one exists for this action
+	TSharedPtr<const FInputChord> HotkeyChord;
 	if (GraphAction->GetTypeId() == FMaterialGraphSchemaAction_NewNode::StaticGetTypeId())
 	{
 		UClass* MaterialExpressionClass = StaticCastSharedPtr<FMaterialGraphSchemaAction_NewNode>(GraphAction)->MaterialExpressionClass;
-		HotkeyGesture = FMaterialEditorSpawnNodeCommands::Get().GetGestureByClass(MaterialExpressionClass);
+		HotkeyChord = FMaterialEditorSpawnNodeCommands::Get().GetChordByClass(MaterialExpressionClass);
 	}
 	else if (GraphAction->GetTypeId() == FMaterialGraphSchemaAction_NewComment::StaticGetTypeId())
 	{
-		HotkeyGesture = FMaterialEditorSpawnNodeCommands::Get().GetGestureByClass(UMaterialExpressionComment::StaticClass());
+		HotkeyChord = FMaterialEditorSpawnNodeCommands::Get().GetChordByClass(UMaterialExpressionComment::StaticClass());
 	}
 
 	// Find icons
@@ -44,7 +44,7 @@ void SMaterialPaletteItem::Construct(const FArguments& InArgs, FCreateWidgetForA
 
 	TSharedRef<SWidget> IconWidget = CreateIconWidget( IconToolTip, IconBrush, IconColor );
 	TSharedRef<SWidget> NameSlotWidget = CreateTextSlotWidget( NameFont, InCreateData, bIsReadOnly );
-	TSharedRef<SWidget> HotkeyDisplayWidget = CreateHotkeyDisplayWidget( NameFont, HotkeyGesture );
+	TSharedRef<SWidget> HotkeyDisplayWidget = CreateHotkeyDisplayWidget( NameFont, HotkeyChord );
 
 	// Create the actual widget
 	this->ChildSlot
@@ -74,12 +74,12 @@ void SMaterialPaletteItem::Construct(const FArguments& InArgs, FCreateWidgetForA
 	];
 }
 
-TSharedRef<SWidget> SMaterialPaletteItem::CreateHotkeyDisplayWidget(const FSlateFontInfo& NameFont, const TSharedPtr<const FInputGesture> HotkeyGesture)
+TSharedRef<SWidget> SMaterialPaletteItem::CreateHotkeyDisplayWidget(const FSlateFontInfo& NameFont, const TSharedPtr<const FInputChord> HotkeyChord)
 {
 	FText HotkeyText;
-	if (HotkeyGesture.IsValid())
+	if (HotkeyChord.IsValid())
 	{
-		HotkeyText = HotkeyGesture->GetInputText();
+		HotkeyText = HotkeyChord->GetInputText();
 	}
 	return SNew(STextBlock)
 		.Text(HotkeyText)
diff --git a/Engine/Source/Editor/MaterialEditor/Private/SMaterialPalette.h b/Engine/Source/Editor/MaterialEditor/Private/SMaterialPalette.h
--- a/Engine/Source/Editor/MaterialEditor/Private/SMaterialPalette.h
+++ b/Engine/Source/Editor/MaterialEditor/Private/SMaterialPalette.h
@@ -16,7 +16,7 @@ public:
 
 private:
 	/* Create the hotkey display widget */
-	TSharedRef<SWidget> CreateHotkeyDisplayWidget(const FSlateFontInfo& NameFont, const TSharedPtr<const FInputGesture> HotkeyGesture);
+	TSharedRef<SWidget> CreateHotkeyDisplayWidget(const FSlateFontInfo& NameFont, const TSharedPtr<const FInputChord> HotkeyChord);
 };
 
 //////////////////////////////////////////////////////////////////////////
diff --git a/Engine/Source/Editor/MaterialEditor/Public/MaterialEditorActions.h b/Engine/Source/Editor/MaterialEditor/Public/MaterialEditorActions.h
--- a/Engine/Source/Editor/MaterialEditor/Public/MaterialEditorActions.h
+++ b/Engine/Source/Editor/MaterialEditor/Public/MaterialEditorActions.h
@@ -176,7 +176,7 @@ public:
 	/** Constructor */
 	FExpressionSpawnInfo(UClass* InMaterialExpressionClass) : MaterialExpressionClass(InMaterialExpressionClass) {}
 
-	/** Holds the UI Command to verify gestures for this action are held */
+	/** Holds the UI Command to verify chords for this action are held */
 	TSharedPtr< FUICommandInfo > CommandInfo;
 
 	/**
@@ -212,14 +212,14 @@ public:
 	virtual void RegisterCommands() override;
 
 	/**
-	 * Returns a graph action assigned to the passed in gesture
+	 * Returns a graph action assigned to the passed in chord
 	 *
-	 * @param InGesture		The gesture to use for lookup
+	 * @param InChord		The chord to use for lookup
 	 * @param InDestGraph	The graph to create the graph action for, used for validation purposes and to link any important node data to the graph
 	 */
-	TSharedPtr< FEdGraphSchemaAction > GetGraphActionByGesture(FInputGesture& InGesture, UEdGraph* InDestGraph) const;
+	TSharedPtr< FEdGraphSchemaAction > GetGraphActionByChord(FInputChord& InChord, UEdGraph* InDestGraph) const;
 
-	const TSharedPtr<const FInputGesture> GetGestureByClass(UClass* MaterialExpressionClass) const;
+	const TSharedPtr<const FInputChord> GetChordByClass(UClass* MaterialExpressionClass) const;
 
 private:
 	/** An array of all the possible commands for spawning nodes */
diff --git a/Engine/Source/Editor/Matinee/Private/MatineeActions.cpp b/Engine/Source/Editor/Matinee/Private/MatineeActions.cpp
--- a/Engine/Source/Editor/Matinee/Private/MatineeActions.cpp
+++ b/Engine/Source/Editor/Matinee/Private/MatineeActions.cpp
@@ -10,184 +10,184 @@
 PRAGMA_DISABLE_OPTIMIZATION
 void FMatineeCommands::RegisterCommands()
 {
-	UI_COMMAND(AddKey, "Add Key", "Add Key", EUserInterfaceActionType::Button, FInputGesture(EModifierKey::None, EKeys::Enter));
+	UI_COMMAND(AddKey, "Add Key", "Add Key", EUserInterfaceActionType::Button, FInputChord(EModifierKey::None, EKeys::Enter));
 	
-	UI_COMMAND(Play, "Play", "Play", EUserInterfaceActionType::Button, FInputGesture(EModifierKey::None, EKeys::L));
-	UI_COMMAND(PlayLoop, "Loop", "Loop Section", EUserInterfaceActionType::Button, FInputGesture());
-	UI_COMMAND(Stop, "Stop", "Stop", EUserInterfaceActionType::Button, FInputGesture(EModifierKey::None, EKeys::K));
-	UI_COMMAND(PlayReverse, "Reverse", "Play in Reverse | Plays the Matinee sequence backwards, starting from the time cursor's position", EUserInterfaceActionType::Button, FInputGesture(EModifierKey::None, EKeys::J));
-
-	UI_COMMAND(PlayPause, "Play/Pause", "Toggle Play/Pause", EUserInterfaceActionType::Button, FInputGesture(EModifierKey::None, EKeys::P));
-
-	UI_COMMAND(CreateCameraActor, "Camera", "Create Camera Actor at Current Camera Location", EUserInterfaceActionType::Button, FInputGesture());
-
-	UI_COMMAND(ToggleSnap, "Snap", "Toggle Snap", EUserInterfaceActionType::ToggleButton, FInputGesture(EModifierKey::None, EKeys::S));
-	UI_COMMAND(ToggleSnapTimeToFrames, "Time to Frames", "Snap Time to Frames | Snaps the timeline cursor to the frame rate specified in the Snap Size setting", EUserInterfaceActionType::ToggleButton, FInputGesture());
-	UI_COMMAND(FixedTimeStepPlayback, "Fixed Time ", "Fixed Time Step Playback | Locks the playback rate to the frame rate specified in the Snap Size setting", EUserInterfaceActionType::ToggleButton, FInputGesture());
-
-	UI_COMMAND(FitSequence, "Sequence", "Fit View to Sequence", EUserInterfaceActionType::Button, FInputGesture(EModifierKey::None, EKeys::A));
-	UI_COMMAND(FitViewToSelected, "Selected", "Fit View to Selected", EUserInterfaceActionType::Button, FInputGesture(EModifierKey::None, EKeys::F));
-	UI_COMMAND(FitLoop, "Loop", "Fit View to Locp", EUserInterfaceActionType::Button, FInputGesture(EModifierKey::Shift, EKeys::A));
-	UI_COMMAND(FitLoopSequence, "Loop Sequence", "Fit View / Loop Sequence", EUserInterfaceActionType::Button, FInputGesture(EModifierKey::Control, EKeys::A));
-	UI_COMMAND(ViewEndofTrack, "End", "Move to End of Track ", EUserInterfaceActionType::Button, FInputGesture(EModifierKey::None, EKeys::E));
-
-	UI_COMMAND(ToggleGorePreview, "Gore", "Enable Gore in Editor Preview", EUserInterfaceActionType::ToggleButton, FInputGesture());
-
-	UI_COMMAND(LaunchRecordWindow, "Record", "Launch Recording Window for Matinee", EUserInterfaceActionType::Button, FInputGesture());
-	UI_COMMAND(CreateMovie, "Movie", "Create a Movie", EUserInterfaceActionType::Button, FInputGesture());
-
-	UI_COMMAND(FileImport, "Import...", "Imports Sequence", EUserInterfaceActionType::Button, FInputGesture());
-	UI_COMMAND(FileExport, "Export All...", "Exports all Sequences", EUserInterfaceActionType::Button, FInputGesture());
-	UI_COMMAND(ExportSoundCueInfo, "Export Sound Cue Info", "Exports Sound Cue info", EUserInterfaceActionType::Button, FInputGesture());
-	UI_COMMAND(ExportAnimInfo, "Export Animation Track Info", "Exports Animation Track info", EUserInterfaceActionType::Button, FInputGesture());
-	UI_COMMAND(FileExportBakeTransforms, "Bake Transforms on Export", "Exports a key every frame, instead of just user created keys", EUserInterfaceActionType::ToggleButton, FInputGesture());
-	UI_COMMAND(FileExportKeepHierarchy, "Keep Hierarchy on Export", "Exports all objects to be in the same hierarchy in engine", EUserInterfaceActionType::ToggleButton, FInputGesture());
-
-	UI_COMMAND(DeleteSelectedKeys, "Delete Keys", "Deletes selected keys", EUserInterfaceActionType::Button, FInputGesture());
-	UI_COMMAND(DuplicateKeys, "Duplicate Keys", "Duplicates selected keys", EUserInterfaceActionType::Button, FInputGesture());
-	UI_COMMAND(InsertSpace, "Insert Space", "Inserts space at the current position", EUserInterfaceActionType::Button, FInputGesture());
-	UI_COMMAND(StretchSection, "Stretch Section", "Stretches the selected section", EUserInterfaceActionType::Button, FInputGesture());
-	UI_COMMAND(StretchSelectedKeyFrames, "Stretch Selected Key Frames", "Stretches the selected keys", EUserInterfaceActionType::Button, FInputGesture());
-	UI_COMMAND(DeleteSection, "Delete Section", "Deletes the selected section", EUserInterfaceActionType::Button, FInputGesture());
-	UI_COMMAND(SelectInSection, "Select In Section", "Selects the keys in Section", EUserInterfaceActionType::Button, FInputGesture());
-	UI_COMMAND(ReduceKeys, "Reduce Keys", "Attempts to remove unnecessary keys", EUserInterfaceActionType::Button, FInputGesture());
-	UI_COMMAND(SavePathTime, "Save Path Time", "Saves the path time at the current position", EUserInterfaceActionType::Button, FInputGesture());
-	UI_COMMAND(JumpToPathTime, "Jump to Path Time", "Jumps to the previously saved path time", EUserInterfaceActionType::Button, FInputGesture());
-
-	UI_COMMAND(ScaleTranslation, "Scale Translation", "Scales the translation of the selected movement trac", EUserInterfaceActionType::Button, FInputGesture() );
-
-	UI_COMMAND(Draw3DTrajectories, "Draw 3D Trajectories", "Toggles 3D trajectories", EUserInterfaceActionType::ToggleButton, FInputGesture());
-	UI_COMMAND(ShowAll3DTrajectories, "Show All 3D Trajectories", "Shows all movement track trajectories", EUserInterfaceActionType::Button, FInputGesture());
-	UI_COMMAND(HideAll3DTrajectories, "Hide All 3D Trajectories", "Hides all movement track trajectories", EUserInterfaceActionType::Button, FInputGesture());
-	UI_COMMAND(PreferFrameNumbers, "Prefer Frame Numbers", "Toggles frame numbers on keys", EUserInterfaceActionType::ToggleButton, FInputGesture());
-	UI_COMMAND(ShowTimeCursorPosForAllKeys, "Show Time Cursor Position", "Shows relative time cursor position on keys", EUserInterfaceActionType::Check, FInputGesture());
-
-	UI_COMMAND(ZoomToTimeCursorPosition, "Zoom To Time Cursor Position", "Toggles center zoom on time cursor position", EUserInterfaceActionType::ToggleButton, FInputGesture());
-	UI_COMMAND(ViewFrameStats, "View Frame Stats", "Views frame stats", EUserInterfaceActionType::Button, FInputGesture());
-	UI_COMMAND(EditingCrosshair, "Editing Crosshair", "Toggles editing crosshair", EUserInterfaceActionType::ToggleButton, FInputGesture());
-
-	UI_COMMAND(EnableEditingGrid,"Enable Editing Grid", "Toggles editing grid", EUserInterfaceActionType::ToggleButton, FInputGesture());
-	UI_COMMAND(TogglePanInvert,"Pan Invert", "Toggles pan invert", EUserInterfaceActionType::ToggleButton, FInputGesture());
-	UI_COMMAND(ToggleAllowKeyframeBarSelection,"Allow Keyframe Bar Selection", "Toggles keyframe bar selection", EUserInterfaceActionType::ToggleButton, FInputGesture());
-	UI_COMMAND(ToggleAllowKeyframeTextSelection,"Allow Keyframe Text Selection", "Toggles keyframe text selection", EUserInterfaceActionType::ToggleButton, FInputGesture());
-
-	UI_COMMAND(ToggleLockCameraPitch,"Lock Camera Pitch", "Toggles lock camera pitch", EUserInterfaceActionType::ToggleButton, FInputGesture());
+	UI_COMMAND(Play, "Play", "Play", EUserInterfaceActionType::Button, FInputChord(EModifierKey::None, EKeys::L));
+	UI_COMMAND(PlayLoop, "Loop", "Loop Section", EUserInterfaceActionType::Button, FInputChord());
+	UI_COMMAND(Stop, "Stop", "Stop", EUserInterfaceActionType::Button, FInputChord(EModifierKey::None, EKeys::K));
+	UI_COMMAND(PlayReverse, "Reverse", "Play in Reverse | Plays the Matinee sequence backwards, starting from the time cursor's position", EUserInterfaceActionType::Button, FInputChord(EModifierKey::None, EKeys::J));
+
+	UI_COMMAND(PlayPause, "Play/Pause", "Toggle Play/Pause", EUserInterfaceActionType::Button, FInputChord(EModifierKey::None, EKeys::P));
+
+	UI_COMMAND(CreateCameraActor, "Camera", "Create Camera Actor at Current Camera Location", EUserInterfaceActionType::Button, FInputChord());
+
+	UI_COMMAND(ToggleSnap, "Snap", "Toggle Snap", EUserInterfaceActionType::ToggleButton, FInputChord(EModifierKey::None, EKeys::S));
+	UI_COMMAND(ToggleSnapTimeToFrames, "Time to Frames", "Snap Time to Frames | Snaps the timeline cursor to the frame rate specified in the Snap Size setting", EUserInterfaceActionType::ToggleButton, FInputChord());
+	UI_COMMAND(FixedTimeStepPlayback, "Fixed Time ", "Fixed Time Step Playback | Locks the playback rate to the frame rate specified in the Snap Size setting", EUserInterfaceActionType::ToggleButton, FInputChord());
+
+	UI_COMMAND(FitSequence, "Sequence", "Fit View to Sequence", EUserInterfaceActionType::Button, FInputChord(EModifierKey::None, EKeys::A));
+	UI_COMMAND(FitViewToSelected, "Selected", "Fit View to Selected", EUserInterfaceActionType::Button, FInputChord(EModifierKey::None, EKeys::F));
+	UI_COMMAND(FitLoop, "Loop", "Fit View to Locp", EUserInterfaceActionType::Button, FInputChord(EModifierKey::Shift, EKeys::A));
+	UI_COMMAND(FitLoopSequence, "Loop Sequence", "Fit View / Loop Sequence", EUserInterfaceActionType::Button, FInputChord(EModifierKey::Control, EKeys::A));
+	UI_COMMAND(ViewEndofTrack, "End", "Move to End of Track ", EUserInterfaceActionType::Button, FInputChord(EModifierKey::None, EKeys::E));
+
+	UI_COMMAND(ToggleGorePreview, "Gore", "Enable Gore in Editor Preview", EUserInterfaceActionType::ToggleButton, FInputChord());
+
+	UI_COMMAND(LaunchRecordWindow, "Record", "Launch Recording Window for Matinee", EUserInterfaceActionType::Button, FInputChord());
+	UI_COMMAND(CreateMovie, "Movie", "Create a Movie", EUserInterfaceActionType::Button, FInputChord());
+
+	UI_COMMAND(FileImport, "Import...", "Imports Sequence", EUserInterfaceActionType::Button, FInputChord());
+	UI_COMMAND(FileExport, "Export All...", "Exports all Sequences", EUserInterfaceActionType::Button, FInputChord());
+	UI_COMMAND(ExportSoundCueInfo, "Export Sound Cue Info", "Exports Sound Cue info", EUserInterfaceActionType::Button, FInputChord());
+	UI_COMMAND(ExportAnimInfo, "Export Animation Track Info", "Exports Animation Track info", EUserInterfaceActionType::Button, FInputChord());
+	UI_COMMAND(FileExportBakeTransforms, "Bake Transforms on Export", "Exports a key every frame, instead of just user created keys", EUserInterfaceActionType::ToggleButton, FInputChord());
+	UI_COMMAND(FileExportKeepHierarchy, "Keep Hierarchy on Export", "Exports all objects to be in the same hierarchy in engine", EUserInterfaceActionType::ToggleButton, FInputChord());
+
+	UI_COMMAND(DeleteSelectedKeys, "Delete Keys", "Deletes selected keys", EUserInterfaceActionType::Button, FInputChord());
+	UI_COMMAND(DuplicateKeys, "Duplicate Keys", "Duplicates selected keys", EUserInterfaceActionType::Button, FInputChord());
+	UI_COMMAND(InsertSpace, "Insert Space", "Inserts space at the current position", EUserInterfaceActionType::Button, FInputChord());
+	UI_COMMAND(StretchSection, "Stretch Section", "Stretches the selected section", EUserInterfaceActionType::Button, FInputChord());
+	UI_COMMAND(StretchSelectedKeyFrames, "Stretch Selected Key Frames", "Stretches the selected keys", EUserInterfaceActionType::Button, FInputChord());
+	UI_COMMAND(DeleteSection, "Delete Section", "Deletes the selected section", EUserInterfaceActionType::Button, FInputChord());
+	UI_COMMAND(SelectInSection, "Select In Section", "Selects the keys in Section", EUserInterfaceActionType::Button, FInputChord());
+	UI_COMMAND(ReduceKeys, "Reduce Keys", "Attempts to remove unnecessary keys", EUserInterfaceActionType::Button, FInputChord());
+	UI_COMMAND(SavePathTime, "Save Path Time", "Saves the path time at the current position", EUserInterfaceActionType::Button, FInputChord());
+	UI_COMMAND(JumpToPathTime, "Jump to Path Time", "Jumps to the previously saved path time", EUserInterfaceActionType::Button, FInputChord());
+
+	UI_COMMAND(ScaleTranslation, "Scale Translation", "Scales the translation of the selected movement trac", EUserInterfaceActionType::Button, FInputChord() );
+
+	UI_COMMAND(Draw3DTrajectories, "Draw 3D Trajectories", "Toggles 3D trajectories", EUserInterfaceActionType::ToggleButton, FInputChord());
+	UI_COMMAND(ShowAll3DTrajectories, "Show All 3D Trajectories", "Shows all movement track trajectories", EUserInterfaceActionType::Button, FInputChord());
+	UI_COMMAND(HideAll3DTrajectories, "Hide All 3D Trajectories", "Hides all movement track trajectories", EUserInterfaceActionType::Button, FInputChord());
+	UI_COMMAND(PreferFrameNumbers, "Prefer Frame Numbers", "Toggles frame numbers on keys", EUserInterfaceActionType::ToggleButton, FInputChord());
+	UI_COMMAND(ShowTimeCursorPosForAllKeys, "Show Time Cursor Position", "Shows relative time cursor position on keys", EUserInterfaceActionType::Check, FInputChord());
+
+	UI_COMMAND(ZoomToTimeCursorPosition, "Zoom To Time Cursor Position", "Toggles center zoom on time cursor position", EUserInterfaceActionType::ToggleButton, FInputChord());
+	UI_COMMAND(ViewFrameStats, "View Frame Stats", "Views frame stats", EUserInterfaceActionType::Button, FInputChord());
+	UI_COMMAND(EditingCrosshair, "Editing Crosshair", "Toggles editing crosshair", EUserInterfaceActionType::ToggleButton, FInputChord());
+
+	UI_COMMAND(EnableEditingGrid,"Enable Editing Grid", "Toggles editing grid", EUserInterfaceActionType::ToggleButton, FInputChord());
+	UI_COMMAND(TogglePanInvert,"Pan Invert", "Toggles pan invert", EUserInterfaceActionType::ToggleButton, FInputChord());
+	UI_COMMAND(ToggleAllowKeyframeBarSelection,"Allow Keyframe Bar Selection", "Toggles keyframe bar selection", EUserInterfaceActionType::ToggleButton, FInputChord());
+	UI_COMMAND(ToggleAllowKeyframeTextSelection,"Allow Keyframe Text Selection", "Toggles keyframe text selection", EUserInterfaceActionType::ToggleButton, FInputChord());
+
+	UI_COMMAND(ToggleLockCameraPitch,"Lock Camera Pitch", "Toggles lock camera pitch", EUserInterfaceActionType::ToggleButton, FInputChord());
 
 	//Context menu commands
-	UI_COMMAND(EditCut,"Cut", "Cuts selected group or track", EUserInterfaceActionType::Button, FInputGesture(EModifierKey::Control, EKeys::X));
-	UI_COMMAND(EditCopy,"Copy", "Copies selected group or track", EUserInterfaceActionType::Button, FInputGesture(EModifierKey::Control, EKeys::C));
-	UI_COMMAND(EditPaste,"Paste", "Pastes group or track", EUserInterfaceActionType::Button, FInputGesture(EModifierKey::Control, EKeys::V));
+	UI_COMMAND(EditCut,"Cut", "Cuts selected group or track", EUserInterfaceActionType::Button, FInputChord(EModifierKey::Control, EKeys::X));
+	UI_COMMAND(EditCopy,"Copy", "Copies selected group or track", EUserInterfaceActionType::Button, FInputChord(EModifierKey::Control, EKeys::C));
+	UI_COMMAND(EditPaste,"Paste", "Pastes group or track", EUserInterfaceActionType::Button, FInputChord(EModifierKey::Control, EKeys::V));
 
 	//Tab
-	UI_COMMAND(GroupDeleteTab, "Delete Group Tab", "Deletes selected group", EUserInterfaceActionType::Button, FInputGesture());
+	UI_COMMAND(GroupDeleteTab, "Delete Group Tab", "Deletes selected group", EUserInterfaceActionType::Button, FInputChord());
 
 	//Group
-	UI_COMMAND(ActorSelectAll, "Select Group Actors","Selects all group actors", EUserInterfaceActionType::Button, FInputGesture());
-	UI_COMMAND(ActorAddAll, "Add Selected Actors","Adds selected actors to group", EUserInterfaceActionType::Button, FInputGesture());
-	UI_COMMAND(ActorReplaceAll, "Replace Group Actors With Selected Actors","Replaces group actors with selected actors", EUserInterfaceActionType::Button, FInputGesture());
-	UI_COMMAND(ActorRemoveAll, "Remove Group Actors","Removes all group actors", EUserInterfaceActionType::Button, FInputGesture());
-	UI_COMMAND(ExportCameraAnim, "Export To CameraAnim","Exports selected group to CameraAnim", EUserInterfaceActionType::Button, FInputGesture());
-	UI_COMMAND(ExportAnimTrackFBX, "Export (.FBX)","Exports selected track as FBX", EUserInterfaceActionType::Button, FInputGesture());
-	UI_COMMAND(ExportAnimGroupFBX, "Export (.FBX)","Exports selected group as FBX", EUserInterfaceActionType::Button, FInputGesture());
-	UI_COMMAND(GroupDuplicate, "Duplicate Group","Duplicates selected group", EUserInterfaceActionType::Button, FInputGesture());
-	UI_COMMAND(GroupDelete, "Delete Group Tab","Deletes selected group", EUserInterfaceActionType::Button, FInputGesture());
-	UI_COMMAND(GroupCreateTab, "Create Group Tab","Creates new group", EUserInterfaceActionType::Button, FInputGesture());
-	UI_COMMAND(GroupRemoveFromTab, "Remove From Selected Group Tab","Removes selected group", EUserInterfaceActionType::Button, FInputGesture());
-	UI_COMMAND(RemoveFromGroupFolder, "Remove From Folder","Moves group to selected folder", EUserInterfaceActionType::Button, FInputGesture());
+	UI_COMMAND(ActorSelectAll, "Select Group Actors","Selects all group actors", EUserInterfaceActionType::Button, FInputChord());
+	UI_COMMAND(ActorAddAll, "Add Selected Actors","Adds selected actors to group", EUserInterfaceActionType::Button, FInputChord());
+	UI_COMMAND(ActorReplaceAll, "Replace Group Actors With Selected Actors","Replaces group actors with selected actors", EUserInterfaceActionType::Button, FInputChord());
+	UI_COMMAND(ActorRemoveAll, "Remove Group Actors","Removes all group actors", EUserInterfaceActionType::Button, FInputChord());
+	UI_COMMAND(ExportCameraAnim, "Export To CameraAnim","Exports selected group to CameraAnim", EUserInterfaceActionType::Button, FInputChord());
+	UI_COMMAND(ExportAnimTrackFBX, "Export (.FBX)","Exports selected track as FBX", EUserInterfaceActionType::Button, FInputChord());
+	UI_COMMAND(ExportAnimGroupFBX, "Export (.FBX)","Exports selected group as FBX", EUserInterfaceActionType::Button, FInputChord());
+	UI_COMMAND(GroupDuplicate, "Duplicate Group","Duplicates selected group", EUserInterfaceActionType::Button, FInputChord());
+	UI_COMMAND(GroupDelete, "Delete Group Tab","Deletes selected group", EUserInterfaceActionType::Button, FInputChord());
+	UI_COMMAND(GroupCreateTab, "Create Group Tab","Creates new group", EUserInterfaceActionType::Button, FInputChord());
+	UI_COMMAND(GroupRemoveFromTab, "Remove From Selected Group Tab","Removes selected group", EUserInterfaceActionType::Button, FInputChord());
+	UI_COMMAND(RemoveFromGroupFolder, "Remove From Folder","Moves group to selected folder", EUserInterfaceActionType::Button, FInputChord());
 
 	//Track
-	UI_COMMAND(TrackRename, "Rename Track","Renames selected track", EUserInterfaceActionType::Button, FInputGesture());
-	UI_COMMAND(TrackDelete, "Delete Track","Deletes selected track", EUserInterfaceActionType::Button, FInputGesture());
-	UI_COMMAND(Show3DTrajectory, "Show 3D Trajectory","Shows 3D trajectory", EUserInterfaceActionType::Button, FInputGesture());
-	UI_COMMAND(TrackSplitTransAndRot, "Split Translation And Rotation","Splits movement into translation and rotation tracks", EUserInterfaceActionType::Button, FInputGesture());
-	UI_COMMAND(TrackNormalizeVelocity, "Normalize Velocity","Normalizes track velocity", EUserInterfaceActionType::Button, FInputGesture());
-	UI_COMMAND(ParticleReplayTrackContextStartRecording, "Start Recording Particles","Starts recording particle replay track", EUserInterfaceActionType::Button, FInputGesture());
-	UI_COMMAND(ParticleReplayTrackContextStopRecording, "Stop Recording Particles","Stops recording particle replay track", EUserInterfaceActionType::Button, FInputGesture());
+	UI_COMMAND(TrackRename, "Rename Track","Renames selected track", EUserInterfaceActionType::Button, FInputChord());
+	UI_COMMAND(TrackDelete, "Delete Track","Deletes selected track", EUserInterfaceActionType::Button, FInputChord());
+	UI_COMMAND(Show3DTrajectory, "Show 3D Trajectory","Shows 3D trajectory", EUserInterfaceActionType::Button, FInputChord());
+	UI_COMMAND(TrackSplitTransAndRot, "Split Translation And Rotation","Splits movement into translation and rotation tracks", EUserInterfaceActionType::Button, FInputChord());
+	UI_COMMAND(TrackNormalizeVelocity, "Normalize Velocity","Normalizes track velocity", EUserInterfaceActionType::Button, FInputChord());
+	UI_COMMAND(ParticleReplayTrackContextStartRecording, "Start Recording Particles","Starts recording particle replay track", EUserInterfaceActionType::Button, FInputChord());
+	UI_COMMAND(ParticleReplayTrackContextStopRecording, "Stop Recording Particles","Stops recording particle replay track", EUserInterfaceActionType::Button, FInputChord());
 
 	//Background
-	UI_COMMAND(NewFolder, "Add New Folder", "Adds new Folder", EUserInterfaceActionType::Button, FInputGesture());
-	UI_COMMAND(NewEmptyGroup, "Add New Empty Group", "Adds new Empty Group", EUserInterfaceActionType::Button, FInputGesture());
-	UI_COMMAND(NewCameraGroup, "Add New Camera Group", "Adds new Camera Group", EUserInterfaceActionType::Button, FInputGesture());
-	UI_COMMAND(NewParticleGroup, "Add New Particle Group", "Adds new Particle Group", EUserInterfaceActionType::Button, FInputGesture());
-	UI_COMMAND(NewSkeletalMeshGroup, "Add New Skeletal Group", "Adds new Skeletal Group", EUserInterfaceActionType::Button, FInputGesture());
-	UI_COMMAND(NewLightingGroup, "Add New Lighting Group", "Adds new Lighting Group", EUserInterfaceActionType::Button, FInputGesture());
-	UI_COMMAND(NewDirectorGroup, "Add New Director Group", "Adds new Director Group", EUserInterfaceActionType::Button, FInputGesture());
+	UI_COMMAND(NewFolder, "Add New Folder", "Adds new Folder", EUserInterfaceActionType::Button, FInputChord());
+	UI_COMMAND(NewEmptyGroup, "Add New Empty Group", "Adds new Empty Group", EUserInterfaceActionType::Button, FInputChord());
+	UI_COMMAND(NewCameraGroup, "Add New Camera Group", "Adds new Camera Group", EUserInterfaceActionType::Button, FInputChord());
+	UI_COMMAND(NewParticleGroup, "Add New Particle Group", "Adds new Particle Group", EUserInterfaceActionType::Button, FInputChord());
+	UI_COMMAND(NewSkeletalMeshGroup, "Add New Skeletal Group", "Adds new Skeletal Group", EUserInterfaceActionType::Button, FInputChord());
+	UI_COMMAND(NewLightingGroup, "Add New Lighting Group", "Adds new Lighting Group", EUserInterfaceActionType::Button, FInputChord());
+	UI_COMMAND(NewDirectorGroup, "Add New Director Group", "Adds new Director Group", EUserInterfaceActionType::Button, FInputChord());
 
-	UI_COMMAND(ToggleCurveEditor, "Curves", "Toggles curve editor", EUserInterfaceActionType::ToggleButton, FInputGesture());
-	UI_COMMAND(ToggleDirectorTimeline, "Director Timeline","Toggles director timeline", EUserInterfaceActionType::ToggleButton, FInputGesture());
+	UI_COMMAND(ToggleCurveEditor, "Curves", "Toggles curve editor", EUserInterfaceActionType::ToggleButton, FInputChord());
+	UI_COMMAND(ToggleDirectorTimeline, "Director Timeline","Toggles director timeline", EUserInterfaceActionType::ToggleButton, FInputChord());
 
 	//Key
-	UI_COMMAND(KeyModeCurveAuto, "Curve (Auto)","Sets mode to Curve (Auto)", EUserInterfaceActionType::Button, FInputGesture());
-	UI_COMMAND(KeyModeCurveAutoClamped, "Curve (Auto/Clamped)","Sets mode to Curve (Auto/Clamped)", EUserInterfaceActionType::Button, FInputGesture());
-	UI_COMMAND(KeyModeCurveBreak, "Curve (Break)","Sets mode to Curve (Break)", EUserInterfaceActionType::Button, FInputGesture());
-	UI_COMMAND(KeyModeLinear, "Linear","Sets mode to Linear", EUserInterfaceActionType::Button, FInputGesture());
-	UI_COMMAND(KeyModeConstant, "Constant","Sets mode to Constant", EUserInterfaceActionType::Button, FInputGesture());
-	UI_COMMAND(KeySetTime, "Set Time","Sets key time", EUserInterfaceActionType::Button, FInputGesture());
-	UI_COMMAND(KeySetValue, "Set Value","Sets key value", EUserInterfaceActionType::Button, FInputGesture());
-	UI_COMMAND(KeySetBool, "Set Bool","Sets key flag", EUserInterfaceActionType::Button, FInputGesture()); 
-	UI_COMMAND(KeySetColor, "Set Color","Sets key color", EUserInterfaceActionType::Button, FInputGesture());
-	UI_COMMAND(MoveKeySetLookup, "Select Transform Lookup Group...","Selects transform lookup group...", EUserInterfaceActionType::Button, FInputGesture());
-	UI_COMMAND(MoveKeyClearLookup, "Clear Transform Lookup Group","Clears transform lookup group", EUserInterfaceActionType::Button, FInputGesture());
-	UI_COMMAND(EventKeyRename, "Rename Event","Renames event key", EUserInterfaceActionType::Button, FInputGesture());
-	UI_COMMAND(DirKeySetTransitionTime, "Set Transition Time","Sets transition time", EUserInterfaceActionType::Button, FInputGesture());
-	UI_COMMAND(DirKeyRenameCameraShot, "Rename Shot","Renames camera key", EUserInterfaceActionType::Button, FInputGesture());
-	UI_COMMAND(KeySetMasterVolume, "Set Master Volume","Sets master volume", EUserInterfaceActionType::Button, FInputGesture());
-	UI_COMMAND(KeySetMasterPitch, "Set Master Pitch","Sets master pitch", EUserInterfaceActionType::Button, FInputGesture());
-	UI_COMMAND(ToggleKeyFlip, "Flip Toggle","Flips toggle action", EUserInterfaceActionType::Button, FInputGesture());
-	UI_COMMAND(KeySetConditionAlways, "Always Active","Toggles always active", EUserInterfaceActionType::ToggleButton, FInputGesture());
-	UI_COMMAND(KeySetConditionGoreEnabled, "Active if Gore is Enabled","Toggles active if gore is enabled", EUserInterfaceActionType::ToggleButton, FInputGesture());
-	UI_COMMAND(KeySetConditionGoreDisabled, "Active if Gore is Disabled","Toggles active if gore is disabled", EUserInterfaceActionType::ToggleButton, FInputGesture());
-	UI_COMMAND(AnimKeyLoop, "Set Looping","Sets key looping", EUserInterfaceActionType::Button, FInputGesture());
-	UI_COMMAND(AnimKeyNoLoop, "Set Non-Looping","Sets key non-looping", EUserInterfaceActionType::Button, FInputGesture());
-	UI_COMMAND(AnimKeySetStartOffset, "Set Start Offset","Sets key start offset", EUserInterfaceActionType::Button, FInputGesture());
-	UI_COMMAND(AnimKeySetEndOffset, "Set End Offset","Sets key end offset", EUserInterfaceActionType::Button, FInputGesture());
-	UI_COMMAND(AnimKeySetPlayRate, "Set Play Rate","Sets key play rate", EUserInterfaceActionType::Button, FInputGesture());
-	UI_COMMAND(AnimKeyToggleReverse, "Reverse","Toggles key reverse", EUserInterfaceActionType::ToggleButton, FInputGesture());
-	UI_COMMAND(KeySyncGenericBrowserToSoundCue, "Sync Browser To Sound Cue","Finds Sound Cue in Content Browser", EUserInterfaceActionType::Button, FInputGesture());
-	UI_COMMAND(ParticleReplayKeySetClipIDNumber, "Select Particle Replay Clip ID","Selects Particle Replay Clip ID", EUserInterfaceActionType::Button, FInputGesture());
-	UI_COMMAND(ParticleReplayKeySetDuration, "Set Particle Replay Duration","Sets Particle Replay Clip duration", EUserInterfaceActionType::Button, FInputGesture());
-	UI_COMMAND(SoundKeySetVolume, "Set Sound Volume","Sets key sound volume", EUserInterfaceActionType::Button, FInputGesture());
-	UI_COMMAND(SoundKeySetPitch, "Set Sound Pitch","Sets key sound pitch", EUserInterfaceActionType::Button, FInputGesture());
+	UI_COMMAND(KeyModeCurveAuto, "Curve (Auto)","Sets mode to Curve (Auto)", EUserInterfaceActionType::Button, FInputChord());
+	UI_COMMAND(KeyModeCurveAutoClamped, "Curve (Auto/Clamped)","Sets mode to Curve (Auto/Clamped)", EUserInterfaceActionType::Button, FInputChord());
+	UI_COMMAND(KeyModeCurveBreak, "Curve (Break)","Sets mode to Curve (Break)", EUserInterfaceActionType::Button, FInputChord());
+	UI_COMMAND(KeyModeLinear, "Linear","Sets mode to Linear", EUserInterfaceActionType::Button, FInputChord());
+	UI_COMMAND(KeyModeConstant, "Constant","Sets mode to Constant", EUserInterfaceActionType::Button, FInputChord());
+	UI_COMMAND(KeySetTime, "Set Time","Sets key time", EUserInterfaceActionType::Button, FInputChord());
+	UI_COMMAND(KeySetValue, "Set Value","Sets key value", EUserInterfaceActionType::Button, FInputChord());
+	UI_COMMAND(KeySetBool, "Set Bool","Sets key flag", EUserInterfaceActionType::Button, FInputChord()); 
+	UI_COMMAND(KeySetColor, "Set Color","Sets key color", EUserInterfaceActionType::Button, FInputChord());
+	UI_COMMAND(MoveKeySetLookup, "Select Transform Lookup Group...","Selects transform lookup group...", EUserInterfaceActionType::Button, FInputChord());
+	UI_COMMAND(MoveKeyClearLookup, "Clear Transform Lookup Group","Clears transform lookup group", EUserInterfaceActionType::Button, FInputChord());
+	UI_COMMAND(EventKeyRename, "Rename Event","Renames event key", EUserInterfaceActionType::Button, FInputChord());
+	UI_COMMAND(DirKeySetTransitionTime, "Set Transition Time","Sets transition time", EUserInterfaceActionType::Button, FInputChord());
+	UI_COMMAND(DirKeyRenameCameraShot, "Rename Shot","Renames camera key", EUserInterfaceActionType::Button, FInputChord());
+	UI_COMMAND(KeySetMasterVolume, "Set Master Volume","Sets master volume", EUserInterfaceActionType::Button, FInputChord());
+	UI_COMMAND(KeySetMasterPitch, "Set Master Pitch","Sets master pitch", EUserInterfaceActionType::Button, FInputChord());
+	UI_COMMAND(ToggleKeyFlip, "Flip Toggle","Flips toggle action", EUserInterfaceActionType::Button, FInputChord());
+	UI_COMMAND(KeySetConditionAlways, "Always Active","Toggles always active", EUserInterfaceActionType::ToggleButton, FInputChord());
+	UI_COMMAND(KeySetConditionGoreEnabled, "Active if Gore is Enabled","Toggles active if gore is enabled", EUserInterfaceActionType::ToggleButton, FInputChord());
+	UI_COMMAND(KeySetConditionGoreDisabled, "Active if Gore is Disabled","Toggles active if gore is disabled", EUserInterfaceActionType::ToggleButton, FInputChord());
+	UI_COMMAND(AnimKeyLoop, "Set Looping","Sets key looping", EUserInterfaceActionType::Button, FInputChord());
+	UI_COMMAND(AnimKeyNoLoop, "Set Non-Looping","Sets key non-looping", EUserInterfaceActionType::Button, FInputChord());
+	UI_COMMAND(AnimKeySetStartOffset, "Set Start Offset","Sets key start offset", EUserInterfaceActionType::Button, FInputChord());
+	UI_COMMAND(AnimKeySetEndOffset, "Set End Offset","Sets key end offset", EUserInterfaceActionType::Button, FInputChord());
+	UI_COMMAND(AnimKeySetPlayRate, "Set Play Rate","Sets key play rate", EUserInterfaceActionType::Button, FInputChord());
+	UI_COMMAND(AnimKeyToggleReverse, "Reverse","Toggles key reverse", EUserInterfaceActionType::ToggleButton, FInputChord());
+	UI_COMMAND(KeySyncGenericBrowserToSoundCue, "Sync Browser To Sound Cue","Finds Sound Cue in Content Browser", EUserInterfaceActionType::Button, FInputChord());
+	UI_COMMAND(ParticleReplayKeySetClipIDNumber, "Select Particle Replay Clip ID","Selects Particle Replay Clip ID", EUserInterfaceActionType::Button, FInputChord());
+	UI_COMMAND(ParticleReplayKeySetDuration, "Set Particle Replay Duration","Sets Particle Replay Clip duration", EUserInterfaceActionType::Button, FInputChord());
+	UI_COMMAND(SoundKeySetVolume, "Set Sound Volume","Sets key sound volume", EUserInterfaceActionType::Button, FInputChord());
+	UI_COMMAND(SoundKeySetPitch, "Set Sound Pitch","Sets key sound pitch", EUserInterfaceActionType::Button, FInputChord());
 
 	//Collapse/Expand context menu
-	UI_COMMAND(ExpandAllGroups, "Expand All","Expand the entire hierarchy of folders, groups and tracks in UnrealMatinee's track window", EUserInterfaceActionType::Button, FInputGesture());
-	UI_COMMAND(CollapseAllGroups, "Collapse All","Collapse the entire hierarchy of folders, groups and tracks in UnrealMatinee's track window", EUserInterfaceActionType::Button, FInputGesture());
+	UI_COMMAND(ExpandAllGroups, "Expand All","Expand the entire hierarchy of folders, groups and tracks in UnrealMatinee's track window", EUserInterfaceActionType::Button, FInputChord());
+	UI_COMMAND(CollapseAllGroups, "Collapse All","Collapse the entire hierarchy of folders, groups and tracks in UnrealMatinee's track window", EUserInterfaceActionType::Button, FInputChord());
 
 	//Marker Context Menu
-	UI_COMMAND(MarkerMoveToBeginning, "Move To Sequence Start","Moves marker to sequence start", EUserInterfaceActionType::Button, FInputGesture());
-	UI_COMMAND(MarkerMoveToEnd, "Move To Sequence End","Moves marker to sequence end", EUserInterfaceActionType::Button, FInputGesture());
-	UI_COMMAND(MarkerMoveToEndOfLongestTrack, "Move To Longest Track Endpoint","Moves marker to end of longest track", EUserInterfaceActionType::Button, FInputGesture());
-	UI_COMMAND(MarkerMoveToEndOfSelectedTrack, "Move To Longest Selected Track Endpoint","Moves marker to end of selected track", EUserInterfaceActionType::Button, FInputGesture());
-	UI_COMMAND(MarkerMoveToCurrentPosition, "Move To Current Timeline Position","Moves marker to current timeline position", EUserInterfaceActionType::Button, FInputGesture());
+	UI_COMMAND(MarkerMoveToBeginning, "Move To Sequence Start","Moves marker to sequence start", EUserInterfaceActionType::Button, FInputChord());
+	UI_COMMAND(MarkerMoveToEnd, "Move To Sequence End","Moves marker to sequence end", EUserInterfaceActionType::Button, FInputChord());
+	UI_COMMAND(MarkerMoveToEndOfLongestTrack, "Move To Longest Track Endpoint","Moves marker to end of longest track", EUserInterfaceActionType::Button, FInputChord());
+	UI_COMMAND(MarkerMoveToEndOfSelectedTrack, "Move To Longest Selected Track Endpoint","Moves marker to end of selected track", EUserInterfaceActionType::Button, FInputChord());
+	UI_COMMAND(MarkerMoveToCurrentPosition, "Move To Current Timeline Position","Moves marker to current timeline position", EUserInterfaceActionType::Button, FInputChord());
 
 	//Viewport/Key commands
-	UI_COMMAND(ZoomIn,"Zoom In","",EUserInterfaceActionType::Button,FInputGesture(EModifierKey::None, EKeys::Add));
-	UI_COMMAND(ZoomOut,"Zoom Out","",EUserInterfaceActionType::Button,FInputGesture(EModifierKey::None, EKeys::Subtract));
-	UI_COMMAND(ZoomInAlt,"Zoom In Alt","",EUserInterfaceActionType::Button,FInputGesture(EModifierKey::None, EKeys::Equals));
-	UI_COMMAND(ZoomOutAlt,"Zoom Out Alt","",EUserInterfaceActionType::Button,FInputGesture(EModifierKey::None, EKeys::Underscore));
-	UI_COMMAND(MarkInSection,"Mark In Section","",EUserInterfaceActionType::Button,FInputGesture(EModifierKey::None, EKeys::I));
-	UI_COMMAND(MarkOutSection,"Mark Out Section","",EUserInterfaceActionType::Button,FInputGesture(EModifierKey::None, EKeys::O));
-	UI_COMMAND(IncrementPosition,"Increment Position","",EUserInterfaceActionType::Button,FInputGesture(EModifierKey::None, EKeys::Right));
-	UI_COMMAND(DecrementPosition,"Decrement Position","",EUserInterfaceActionType::Button,FInputGesture(EModifierKey::None, EKeys::Left));
-	UI_COMMAND(MoveToNextKey,"Move To Next Key","",EUserInterfaceActionType::Button,FInputGesture(EModifierKey::Shift, EKeys::Right));
-	UI_COMMAND(MoveToPrevKey,"Move To Previous Key","",EUserInterfaceActionType::Button,FInputGesture(EModifierKey::Shift, EKeys::Left));
-	UI_COMMAND(SplitAnimKey,"Split Anim Key","",EUserInterfaceActionType::Button,FInputGesture(EModifierKey::None, EKeys::R));
-	UI_COMMAND(MoveActiveUp,"Move Active Track/Group Up","",EUserInterfaceActionType::Button,FInputGesture(EModifierKey::None, EKeys::Up));
-	UI_COMMAND(MoveActiveDown,"Move Active Track/Group Down","",EUserInterfaceActionType::Button,FInputGesture(EModifierKey::None, EKeys::Down));
+	UI_COMMAND(ZoomIn,"Zoom In","",EUserInterfaceActionType::Button,FInputChord(EModifierKey::None, EKeys::Add));
+	UI_COMMAND(ZoomOut,"Zoom Out","",EUserInterfaceActionType::Button,FInputChord(EModifierKey::None, EKeys::Subtract));
+	UI_COMMAND(ZoomInAlt,"Zoom In Alt","",EUserInterfaceActionType::Button,FInputChord(EModifierKey::None, EKeys::Equals));
+	UI_COMMAND(ZoomOutAlt,"Zoom Out Alt","",EUserInterfaceActionType::Button,FInputChord(EModifierKey::None, EKeys::Underscore));
+	UI_COMMAND(MarkInSection,"Mark In Section","",EUserInterfaceActionType::Button,FInputChord(EModifierKey::None, EKeys::I));
+	UI_COMMAND(MarkOutSection,"Mark Out Section","",EUserInterfaceActionType::Button,FInputChord(EModifierKey::None, EKeys::O));
+	UI_COMMAND(IncrementPosition,"Increment Position","",EUserInterfaceActionType::Button,FInputChord(EModifierKey::None, EKeys::Right));
+	UI_COMMAND(DecrementPosition,"Decrement Position","",EUserInterfaceActionType::Button,FInputChord(EModifierKey::None, EKeys::Left));
+	UI_COMMAND(MoveToNextKey,"Move To Next Key","",EUserInterfaceActionType::Button,FInputChord(EModifierKey::Shift, EKeys::Right));
+	UI_COMMAND(MoveToPrevKey,"Move To Previous Key","",EUserInterfaceActionType::Button,FInputChord(EModifierKey::Shift, EKeys::Left));
+	UI_COMMAND(SplitAnimKey,"Split Anim Key","",EUserInterfaceActionType::Button,FInputChord(EModifierKey::None, EKeys::R));
+	UI_COMMAND(MoveActiveUp,"Move Active Track/Group Up","",EUserInterfaceActionType::Button,FInputChord(EModifierKey::None, EKeys::Up));
+	UI_COMMAND(MoveActiveDown,"Move Active Track/Group Down","",EUserInterfaceActionType::Button,FInputChord(EModifierKey::None, EKeys::Down));
 #if PLATFORM_MAC
-	UI_COMMAND(DuplicateSelectedKeys,"Duplicate Selected Keys","",EUserInterfaceActionType::Button,FInputGesture(EModifierKey::Command, EKeys::W));
+	UI_COMMAND(DuplicateSelectedKeys,"Duplicate Selected Keys","",EUserInterfaceActionType::Button,FInputChord(EModifierKey::Command, EKeys::W));
 #else
-	UI_COMMAND(DuplicateSelectedKeys,"Duplicate Selected Keys","",EUserInterfaceActionType::Button,FInputGesture(EModifierKey::Control, EKeys::W));
+	UI_COMMAND(DuplicateSelectedKeys,"Duplicate Selected Keys","",EUserInterfaceActionType::Button,FInputChord(EModifierKey::Control, EKeys::W));
 #endif
-	UI_COMMAND(CropAnimationBeginning,"Crop Animation Start","",EUserInterfaceActionType::Button,FInputGesture(EModifierKey::Control, EKeys::I));
-	UI_COMMAND(CropAnimationEnd,"Crop Animation End","",EUserInterfaceActionType::Button,FInputGesture(EModifierKey::Control, EKeys::O));
-	UI_COMMAND(ChangeKeyInterpModeAuto,"Change Key Interp Mode Auto","",EUserInterfaceActionType::Button,FInputGesture(EModifierKey::None, EKeys::One));
-	UI_COMMAND(ChangeKeyInterpModeAutoClamped,"Change Key Interp Mode Auto Clamped","",EUserInterfaceActionType::Button,FInputGesture());
-	UI_COMMAND(ChangeKeyInterpModeUser,"Change Key Interp Mode User","",EUserInterfaceActionType::Button,FInputGesture(EModifierKey::None, EKeys::Two));
-	UI_COMMAND(ChangeKeyInterpModeBreak,"Change Key Interp Mode Break","",EUserInterfaceActionType::Button,FInputGesture(EModifierKey::None, EKeys::Three));
-	UI_COMMAND(ChangeKeyInterpModeLinear,"Change Key Interp Mode Linear","",EUserInterfaceActionType::Button,FInputGesture(EModifierKey::None, EKeys::Four));
-	UI_COMMAND(ChangeKeyInterpModeConstant,"Change Key Interp Mode Constant","",EUserInterfaceActionType::Button,FInputGesture(EModifierKey::None, EKeys::Five));
-	UI_COMMAND(DeleteSelection, "Delete Selection", "", EUserInterfaceActionType::Button, FInputGesture(EModifierKey::None, EKeys::Platform_Delete));
+	UI_COMMAND(CropAnimationBeginning,"Crop Animation Start","",EUserInterfaceActionType::Button,FInputChord(EModifierKey::Control, EKeys::I));
+	UI_COMMAND(CropAnimationEnd,"Crop Animation End","",EUserInterfaceActionType::Button,FInputChord(EModifierKey::Control, EKeys::O));
+	UI_COMMAND(ChangeKeyInterpModeAuto,"Change Key Interp Mode Auto","",EUserInterfaceActionType::Button,FInputChord(EModifierKey::None, EKeys::One));
+	UI_COMMAND(ChangeKeyInterpModeAutoClamped,"Change Key Interp Mode Auto Clamped","",EUserInterfaceActionType::Button,FInputChord());
+	UI_COMMAND(ChangeKeyInterpModeUser,"Change Key Interp Mode User","",EUserInterfaceActionType::Button,FInputChord(EModifierKey::None, EKeys::Two));
+	UI_COMMAND(ChangeKeyInterpModeBreak,"Change Key Interp Mode Break","",EUserInterfaceActionType::Button,FInputChord(EModifierKey::None, EKeys::Three));
+	UI_COMMAND(ChangeKeyInterpModeLinear,"Change Key Interp Mode Linear","",EUserInterfaceActionType::Button,FInputChord(EModifierKey::None, EKeys::Four));
+	UI_COMMAND(ChangeKeyInterpModeConstant,"Change Key Interp Mode Constant","",EUserInterfaceActionType::Button,FInputChord(EModifierKey::None, EKeys::Five));
+	UI_COMMAND(DeleteSelection, "Delete Selection", "", EUserInterfaceActionType::Button, FInputChord(EModifierKey::None, EKeys::Platform_Delete));
 }
 PRAGMA_ENABLE_OPTIMIZATION
 
