When editing a CameraAnim, display the anim object's properties in the details panel by default.

diff --git a/Engine/Source/Editor/Matinee/Private/MatineePropertyWindow.cpp b/Engine/Source/Editor/Matinee/Private/MatineePropertyWindow.cpp
--- a/Engine/Source/Editor/Matinee/Private/MatineePropertyWindow.cpp
+++ b/Engine/Source/Editor/Matinee/Private/MatineePropertyWindow.cpp
@@ -12,6 +12,18 @@ void FMatinee::BuildPropertyWindow()
 
 	FPropertyEditorModule& PropertyModule = FModuleManager::LoadModuleChecked<FPropertyEditorModule>("PropertyEditor");
 	PropertyWindow = PropertyModule.CreateDetailView( Args );
+
+	if (IsCameraAnim())
+	{
+		// if editing a CameraAnim, show the CameraAnim object properties by default (when no track or group selected)
+		AMatineeActorCameraAnim* const CamAnimMatineeActor = Cast<AMatineeActorCameraAnim>(MatineeActor);
+		if ((CamAnimMatineeActor != nullptr) && (CamAnimMatineeActor->CameraAnim != nullptr))
+		{
+			TArray<UObject*> Objects;
+			Objects.Add(CamAnimMatineeActor->CameraAnim);
+			PropertyWindow->SetObjects(Objects);
+		}
+	}
 }
 
 
@@ -50,6 +62,15 @@ void FMatinee::UpdatePropertyWindow()
 			Objects.Add(SelectedGroups[Index]);
 		}
 	}
+	else if (IsCameraAnim())
+	{
+		// if editing a CameraAnim, show the CameraAnim properties when nothing else selected
+		AMatineeActorCameraAnim* const CamAnimMatineeActor = Cast<AMatineeActorCameraAnim>(MatineeActor);
+		if ((CamAnimMatineeActor != nullptr) && (CamAnimMatineeActor->CameraAnim != nullptr))
+		{
+			Objects.Add(CamAnimMatineeActor->CameraAnim);
+		}
+	}
 	// else send nothing
 	
 	PropertyWindow->SetObjects(Objects);
