Fixed nonunity build

diff --git a/Engine/Source/Editor/Matinee/Private/MatineePropertyWindow.cpp b/Engine/Source/Editor/Matinee/Private/MatineePropertyWindow.cpp
--- a/Engine/Source/Editor/Matinee/Private/MatineePropertyWindow.cpp
+++ b/Engine/Source/Editor/Matinee/Private/MatineePropertyWindow.cpp
@@ -2,6 +2,7 @@
 
 #include "MatineeModule.h"
 #include "Matinee.h"
+#include "Matinee/MatineeActorCameraAnim.h"
 #include "Editor/PropertyEditor/Public/PropertyEditorModule.h"
 #include "Editor/PropertyEditor/Public/IDetailsView.h"
 
