Added config migration path for newer versions of the engine.

diff --git a/Engine/Source/Editor/Matinee/Private/Matinee.cpp b/Engine/Source/Editor/Matinee/Private/Matinee.cpp
--- a/Engine/Source/Editor/Matinee/Private/Matinee.cpp
+++ b/Engine/Source/Editor/Matinee/Private/Matinee.cpp
@@ -353,7 +353,7 @@ void FMatinee::OnToggleAspectRatioBars()
 			bool bEnabled = !AreAspectRatioBarsEnabled();
 			GCurrentLevelEditingViewportClient->SetShowAspectRatioBarDisplay(bEnabled);
 
-			GConfig->SetBool(TEXT("Matinee"), TEXT("AspectRatioBars"), bEnabled, GEditorUserSettingsIni);
+			GConfig->SetBool(TEXT("Matinee"), TEXT("AspectRatioBars"), bEnabled, GEditorPerProjectIni);
 		}
 	}
 }
@@ -367,7 +367,7 @@ void FMatinee::OnToggleSafeFrames()
 			bool bEnabled = !IsSafeFrameDisplayEnabled();
 			GCurrentLevelEditingViewportClient->SetShowSafeFrameBoxDisplay(bEnabled);
 
-			GConfig->SetBool(TEXT("Matinee"), TEXT("SafeFrames"), bEnabled, GEditorUserSettingsIni);
+			GConfig->SetBool(TEXT("Matinee"), TEXT("SafeFrames"), bEnabled, GEditorPerProjectIni);
 		}
 	}
 }
@@ -375,7 +375,7 @@ void FMatinee::OnToggleSafeFrames()
 bool FMatinee::AreAspectRatioBarsEnabled() const
 {
 	bool bEnabled = false;
-	if ( !GConfig->GetBool(TEXT("Matinee"), TEXT("AspectRatioBars"), bEnabled, GEditorUserSettingsIni) )
+	if ( !GConfig->GetBool(TEXT("Matinee"), TEXT("AspectRatioBars"), bEnabled, GEditorPerProjectIni) )
 	{
 		// We enable them by default
 		return true;
@@ -387,7 +387,7 @@ bool FMatinee::AreAspectRatioBarsEnabled() const
 bool FMatinee::IsSafeFrameDisplayEnabled() const
 {
 	bool bEnabled = false;
-	if ( !GConfig->GetBool(TEXT("Matinee"), TEXT("SafeFrames"), bEnabled, GEditorUserSettingsIni) )
+	if ( !GConfig->GetBool(TEXT("Matinee"), TEXT("SafeFrames"), bEnabled, GEditorPerProjectIni) )
 	{
 		// We not enabled by default
 		return false;
@@ -454,37 +454,37 @@ void FMatinee::InitMatinee(const EToolkitMode::Type Mode, const TSharedPtr< clas
 
 	// 3D tracks should be visible by default
 	bHide3DTrackView = false;
-	GConfig->GetBool( TEXT("Matinee"), TEXT("Hide3DTracks"), bHide3DTrackView, GEditorUserSettingsIni );
+	GConfig->GetBool( TEXT("Matinee"), TEXT("Hide3DTracks"), bHide3DTrackView, GEditorPerProjectIni );
 
 	// Zoom to scrub position defaults to off.  We want zoom to cursor position by default.
 	bZoomToScrubPos = false;
-	GConfig->GetBool( TEXT("Matinee"), TEXT("ZoomToScrubPos"), bZoomToScrubPos, GEditorUserSettingsIni );
+	GConfig->GetBool( TEXT("Matinee"), TEXT("ZoomToScrubPos"), bZoomToScrubPos, GEditorPerProjectIni );
 
 	// Setup 'viewport frame stats' preference
 	bViewportFrameStatsEnabled = true;
-	GConfig->GetBool( TEXT("Matinee"), TEXT("ViewportFrameStats"), bViewportFrameStatsEnabled, GEditorUserSettingsIni );
+	GConfig->GetBool( TEXT("Matinee"), TEXT("ViewportFrameStats"), bViewportFrameStatsEnabled, GEditorPerProjectIni );
 
 	// Get the editing grid size from user settings
 	EditingGridSize = 1;
-	GConfig->GetInt( TEXT("Matinee"), TEXT("EditingGridSize"), EditingGridSize, GEditorUserSettingsIni );
+	GConfig->GetInt( TEXT("Matinee"), TEXT("EditingGridSize"), EditingGridSize, GEditorPerProjectIni );
 
 	// Look to see if the crosshair should be enabled
 	// Disabled by default
 	bEditingCrosshairEnabled = false;
-	GConfig->GetBool( TEXT("Matinee"), TEXT("EditingCrosshair"), bEditingCrosshairEnabled, GEditorUserSettingsIni );
+	GConfig->GetBool( TEXT("Matinee"), TEXT("EditingCrosshair"), bEditingCrosshairEnabled, GEditorPerProjectIni );
 
 	// Look to see if the editing grid should be enabled
 	bEditingGridEnabled = false;
-	GConfig->GetBool( TEXT("Matinee"), TEXT("EnableEditingGrid"), bEditingGridEnabled, GEditorUserSettingsIni );
+	GConfig->GetBool( TEXT("Matinee"), TEXT("EnableEditingGrid"), bEditingGridEnabled, GEditorPerProjectIni );
 
 	// Setup "allow keyframe bar selection" preference
-	GConfig->GetBool( TEXT("Matinee"), TEXT("AllowKeyframeBarSelection"), bAllowKeyframeBarSelection, GEditorUserSettingsIni ); 
+	GConfig->GetBool( TEXT("Matinee"), TEXT("AllowKeyframeBarSelection"), bAllowKeyframeBarSelection, GEditorPerProjectIni ); 
 
 	// Setup "allow keyframe text selection" preference
-	GConfig->GetBool( TEXT("Matinee"), TEXT("AllowKeyframeTextSelection"), bAllowKeyframeTextSelection, GEditorUserSettingsIni );
+	GConfig->GetBool( TEXT("Matinee"), TEXT("AllowKeyframeTextSelection"), bAllowKeyframeTextSelection, GEditorPerProjectIni );
 
 	bInvertPan = true;
-	GConfig->GetBool( TEXT("Matinee"), TEXT("InterpEdPanInvert"), bInvertPan, GEditorUserSettingsIni );
+	GConfig->GetBool( TEXT("Matinee"), TEXT("InterpEdPanInvert"), bInvertPan, GEditorPerProjectIni );
 
 	// Setup "lock camera pitch" preference
 	GetLockCameraPitchFromConfig();
@@ -876,23 +876,23 @@ void FMatinee::InitMatinee(const EToolkitMode::Type Mode, const TSharedPtr< clas
 	// Restore the director timeline setting
 	if ( DirectorTrackWindow.IsValid() && DirectorTrackWindow->InterpEdVC.IsValid() )
 	{
-		GConfig->GetBool(TEXT("Matinee"), TEXT("DirectorTimelineEnabled"), DirectorTrackWindow->InterpEdVC->bWantTimeline, GEditorUserSettingsIni);
+		GConfig->GetBool(TEXT("Matinee"), TEXT("DirectorTimelineEnabled"), DirectorTrackWindow->InterpEdVC->bWantTimeline, GEditorPerProjectIni);
 	}
 
 	// Load fixed time step setting
-	GConfig->GetBool( TEXT("Matinee"), TEXT("FixedTimeStepPlayback"), bFixedTimeStepPlayback, GEditorUserSettingsIni );
+	GConfig->GetBool( TEXT("Matinee"), TEXT("FixedTimeStepPlayback"), bFixedTimeStepPlayback, GEditorPerProjectIni );
 
 	// Load 'prefer frame numbers' setting
-	GConfig->GetBool( TEXT("Matinee"), TEXT("PreferFrameNumbers"), bPreferFrameNumbers, GEditorUserSettingsIni );
+	GConfig->GetBool( TEXT("Matinee"), TEXT("PreferFrameNumbers"), bPreferFrameNumbers, GEditorPerProjectIni );
 
 	// Load 'show time cursor pos for all keys' setting
-	GConfig->GetBool( TEXT("Matinee"), TEXT("ShowTimeCursorPosForAllKeys"), bShowTimeCursorPosForAllKeys, GEditorUserSettingsIni );
+	GConfig->GetBool( TEXT("Matinee"), TEXT("ShowTimeCursorPosForAllKeys"), bShowTimeCursorPosForAllKeys, GEditorPerProjectIni );
 
 	// Restore selected snap mode from INI.
-	GConfig->GetBool( TEXT("Matinee"), TEXT("SnapEnabled"), bSnapEnabled, GEditorUserSettingsIni );
-	GConfig->GetBool( TEXT("Matinee"), TEXT("SnapTimeToFrames"), bSnapTimeToFrames, GEditorUserSettingsIni );
+	GConfig->GetBool( TEXT("Matinee"), TEXT("SnapEnabled"), bSnapEnabled, GEditorPerProjectIni );
+	GConfig->GetBool( TEXT("Matinee"), TEXT("SnapTimeToFrames"), bSnapTimeToFrames, GEditorPerProjectIni );
 	int32 SelectedSnapMode = 3; // default 0.5 sec
-	GConfig->GetInt(TEXT("Matinee"), TEXT("SelectedSnapMode"), SelectedSnapMode, GEditorUserSettingsIni );
+	GConfig->GetInt(TEXT("Matinee"), TEXT("SelectedSnapMode"), SelectedSnapMode, GEditorPerProjectIni );
 
 	//OnChangeSnapSize(SelectedSnapMode);
 	SnapCombo->SetSelectedItem(SnapComboStrings[SelectedSnapMode]);
@@ -911,7 +911,7 @@ void FMatinee::InitMatinee(const EToolkitMode::Type Mode, const TSharedPtr< clas
 	{
 		// NOTE: InitialInterpMode now has a '2' suffix after a version bump to change the default
 		int32 DesiredInitialInterpMode = ( int32 )InitialInterpMode;
-		GConfig->GetInt( TEXT( "Matinee" ), TEXT( "InitialInterpMode2" ), DesiredInitialInterpMode, GEditorUserSettingsIni );
+		GConfig->GetInt( TEXT( "Matinee" ), TEXT( "InitialInterpMode2" ), DesiredInitialInterpMode, GEditorPerProjectIni );
 		InitialInterpModeComboBox->SetSelectedItem( InitialInterpModeStrings[DesiredInitialInterpMode] );
 	}
 
@@ -1874,48 +1874,48 @@ const double FMatinee::GetRecordingEndTime (void) const
 /** Save record settings for next run */
 void FMatinee::SaveRecordingSettings(const FCameraControllerConfig& InCameraConfig)
 {
-	GConfig->SetInt(TEXT("InterpEd.Recording"), TEXT("Mode"), RecordMode, GEditorUserSettingsIni);
+	GConfig->SetInt(TEXT("InterpEd.Recording"), TEXT("Mode"), RecordMode, GEditorPerProjectIni);
 
-	GConfig->SetFloat(TEXT("InterpEd.Recording"), TEXT("TranslationSpeed"), InCameraConfig.TranslationMultiplier, GEditorUserSettingsIni);
-	GConfig->SetFloat(TEXT("InterpEd.Recording"), TEXT("RotationSpeed"), InCameraConfig.RotationMultiplier, GEditorUserSettingsIni);
-	GConfig->SetFloat(TEXT("InterpEd.Recording"), TEXT("ZoomSpeed"), InCameraConfig.ZoomMultiplier, GEditorUserSettingsIni);
+	GConfig->SetFloat(TEXT("InterpEd.Recording"), TEXT("TranslationSpeed"), InCameraConfig.TranslationMultiplier, GEditorPerProjectIni);
+	GConfig->SetFloat(TEXT("InterpEd.Recording"), TEXT("RotationSpeed"), InCameraConfig.RotationMultiplier, GEditorPerProjectIni);
+	GConfig->SetFloat(TEXT("InterpEd.Recording"), TEXT("ZoomSpeed"), InCameraConfig.ZoomMultiplier, GEditorPerProjectIni);
 
-	GConfig->SetBool(TEXT("InterpEd.Recording"), TEXT("InvertX"), InCameraConfig.bInvertX, GEditorUserSettingsIni);
-	GConfig->SetBool(TEXT("InterpEd.Recording"), TEXT("InvertY"), InCameraConfig.bInvertY, GEditorUserSettingsIni);
+	GConfig->SetBool(TEXT("InterpEd.Recording"), TEXT("InvertX"), InCameraConfig.bInvertX, GEditorPerProjectIni);
+	GConfig->SetBool(TEXT("InterpEd.Recording"), TEXT("InvertY"), InCameraConfig.bInvertY, GEditorPerProjectIni);
 	
-	GConfig->SetInt(TEXT("InterpEd.Recording"), TEXT("RollSamples"), RecordRollSmoothingSamples, GEditorUserSettingsIni);
-	GConfig->SetInt(TEXT("InterpEd.Recording"), TEXT("PitchSamples"), RecordPitchSmoothingSamples, GEditorUserSettingsIni);
+	GConfig->SetInt(TEXT("InterpEd.Recording"), TEXT("RollSamples"), RecordRollSmoothingSamples, GEditorPerProjectIni);
+	GConfig->SetInt(TEXT("InterpEd.Recording"), TEXT("PitchSamples"), RecordPitchSmoothingSamples, GEditorPerProjectIni);
 
-	GConfig->SetInt(TEXT("InterpEd.Recording"), TEXT("CameraMovement"), RecordCameraMovementScheme, GEditorUserSettingsIni);
+	GConfig->SetInt(TEXT("InterpEd.Recording"), TEXT("CameraMovement"), RecordCameraMovementScheme, GEditorPerProjectIni);
 
 	FLevelEditorViewportClient* LevelVC = GetRecordingViewport();
 	if (LevelVC)
 	{
-		GConfig->SetFloat(TEXT("InterpEd.Recording"), TEXT("ZoomDistance"), LevelVC->ViewFOV, GEditorUserSettingsIni);
+		GConfig->SetFloat(TEXT("InterpEd.Recording"), TEXT("ZoomDistance"), LevelVC->ViewFOV, GEditorPerProjectIni);
 	}
 }
 
 /** Load record settings for next run */
 void FMatinee::LoadRecordingSettings(FCameraControllerConfig& InCameraConfig)
 {
-	GConfig->GetInt(TEXT("InterpEd.Recording"), TEXT("Mode"), RecordMode, GEditorUserSettingsIni);
+	GConfig->GetInt(TEXT("InterpEd.Recording"), TEXT("Mode"), RecordMode, GEditorPerProjectIni);
 
-	GConfig->GetFloat(TEXT("InterpEd.Recording"), TEXT("TranslationSpeed"), InCameraConfig.TranslationMultiplier, GEditorUserSettingsIni);
-	GConfig->GetFloat(TEXT("InterpEd.Recording"), TEXT("RotationSpeed"), InCameraConfig.RotationMultiplier, GEditorUserSettingsIni);
-	GConfig->GetFloat(TEXT("InterpEd.Recording"), TEXT("ZoomSpeed"), InCameraConfig.ZoomMultiplier, GEditorUserSettingsIni);
+	GConfig->GetFloat(TEXT("InterpEd.Recording"), TEXT("TranslationSpeed"), InCameraConfig.TranslationMultiplier, GEditorPerProjectIni);
+	GConfig->GetFloat(TEXT("InterpEd.Recording"), TEXT("RotationSpeed"), InCameraConfig.RotationMultiplier, GEditorPerProjectIni);
+	GConfig->GetFloat(TEXT("InterpEd.Recording"), TEXT("ZoomSpeed"), InCameraConfig.ZoomMultiplier, GEditorPerProjectIni);
 
-	GConfig->GetBool(TEXT("InterpEd.Recording"), TEXT("InvertX"), InCameraConfig.bInvertX, GEditorUserSettingsIni);
-	GConfig->GetBool(TEXT("InterpEd.Recording"), TEXT("InvertY"), InCameraConfig.bInvertY, GEditorUserSettingsIni);
+	GConfig->GetBool(TEXT("InterpEd.Recording"), TEXT("InvertX"), InCameraConfig.bInvertX, GEditorPerProjectIni);
+	GConfig->GetBool(TEXT("InterpEd.Recording"), TEXT("InvertY"), InCameraConfig.bInvertY, GEditorPerProjectIni);
 
-	GConfig->GetInt(TEXT("InterpEd.Recording"), TEXT("RollSamples"), RecordRollSmoothingSamples, GEditorUserSettingsIni);
-	GConfig->GetInt(TEXT("InterpEd.Recording"), TEXT("PitchSamples"), RecordPitchSmoothingSamples, GEditorUserSettingsIni);
+	GConfig->GetInt(TEXT("InterpEd.Recording"), TEXT("RollSamples"), RecordRollSmoothingSamples, GEditorPerProjectIni);
+	GConfig->GetInt(TEXT("InterpEd.Recording"), TEXT("PitchSamples"), RecordPitchSmoothingSamples, GEditorPerProjectIni);
 
-	GConfig->GetInt(TEXT("InterpEd.Recording"), TEXT("CameraMovement"), RecordCameraMovementScheme, GEditorUserSettingsIni);
+	GConfig->GetInt(TEXT("InterpEd.Recording"), TEXT("CameraMovement"), RecordCameraMovementScheme, GEditorPerProjectIni);
 
 	FLevelEditorViewportClient* LevelVC = GetRecordingViewport();
 	if (LevelVC)
 	{
-		GConfig->GetFloat(TEXT("InterpEd.Recording"), TEXT("ZoomDistance"), LevelVC->ViewFOV, GEditorUserSettingsIni);
+		GConfig->GetFloat(TEXT("InterpEd.Recording"), TEXT("ZoomDistance"), LevelVC->ViewFOV, GEditorPerProjectIni);
 	}
 }
 
diff --git a/Engine/Source/Editor/Matinee/Private/MatineeMenus.cpp b/Engine/Source/Editor/Matinee/Private/MatineeMenus.cpp
--- a/Engine/Source/Editor/Matinee/Private/MatineeMenus.cpp
+++ b/Engine/Source/Editor/Matinee/Private/MatineeMenus.cpp
@@ -1051,7 +1051,7 @@ void FMatinee::OnViewHide3DTracks()
 {
 	bHide3DTrackView = !bHide3DTrackView;
 	// Save to ini when it changes.
-	GConfig->SetBool( TEXT("Matinee"), TEXT("Hide3DTracks"), bHide3DTrackView, GEditorUserSettingsIni );
+	GConfig->SetBool( TEXT("Matinee"), TEXT("Hide3DTracks"), bHide3DTrackView, GEditorPerProjectIni );
 }
 
 bool FMatinee::IsViewHide3DTracksToggled()
@@ -1065,7 +1065,7 @@ void FMatinee::OnViewZoomToScrubPos()
 	bZoomToScrubPos = !bZoomToScrubPos;
 	
 	// Save to ini when it changes.
-	GConfig->SetBool( TEXT("Matinee"), TEXT("ZoomToScrubPos"), bZoomToScrubPos, GEditorUserSettingsIni );
+	GConfig->SetBool( TEXT("Matinee"), TEXT("ZoomToScrubPos"), bZoomToScrubPos, GEditorPerProjectIni );
 }
 
 bool FMatinee::IsViewZoomToScrubPosToggled()
@@ -1076,7 +1076,7 @@ bool FMatinee::IsViewZoomToScrubPosToggled()
 void FMatinee::OnEnableEditingGrid()
 {
 	bEditingGridEnabled = !bEditingGridEnabled;
-	GConfig->SetBool( TEXT("Matinee"), TEXT("EnableEditingGrid"), bEditingGridEnabled, GEditorUserSettingsIni );
+	GConfig->SetBool( TEXT("Matinee"), TEXT("EnableEditingGrid"), bEditingGridEnabled, GEditorPerProjectIni );
 }
 
 bool FMatinee::IsEnableEditingGridToggled()
@@ -1088,7 +1088,7 @@ void FMatinee::OnSetEditingGrid( uint32 InGridSize )
 {
 	EditingGridSize = InGridSize;
 
-	GConfig->SetInt( TEXT("Matinee"), TEXT("EditingGridSize"), EditingGridSize, GEditorUserSettingsIni );
+	GConfig->SetInt( TEXT("Matinee"), TEXT("EditingGridSize"), EditingGridSize, GEditorPerProjectIni );
 }
 
 bool FMatinee::IsEditingGridChecked(uint32 InGridSize)
@@ -1099,7 +1099,7 @@ bool FMatinee::IsEditingGridChecked(uint32 InGridSize)
 void FMatinee::OnToggleEditingCrosshair()
 {
 	bEditingCrosshairEnabled = !bEditingCrosshairEnabled;
-	GConfig->SetBool( TEXT("Matinee"), TEXT("EditingCrosshair"), bEditingCrosshairEnabled, GEditorUserSettingsIni );
+	GConfig->SetBool( TEXT("Matinee"), TEXT("EditingCrosshair"), bEditingCrosshairEnabled, GEditorPerProjectIni );
 }
 
 bool FMatinee::IsEditingCrosshairToggled()
@@ -1112,7 +1112,7 @@ void FMatinee::OnToggleViewportFrameStats()
 	bViewportFrameStatsEnabled = !bViewportFrameStatsEnabled;
 
 	// Save to ini when it changes.
-	GConfig->SetBool( TEXT("Matinee"), TEXT("ViewportFrameStats"), bViewportFrameStatsEnabled, GEditorUserSettingsIni );
+	GConfig->SetBool( TEXT("Matinee"), TEXT("ViewportFrameStats"), bViewportFrameStatsEnabled, GEditorPerProjectIni );
 }
 
 bool FMatinee::IsViewportFrameStatsToggled()
@@ -1585,7 +1585,7 @@ void FMatinee::OnContextTrackExportAnimFBX()
 					// Export the Matinee information to a COLLADA document.
 					Exporter->CreateDocument();
 					Exporter->SetTrasformBaking(bBakeTransforms);
-					const bool bKeepHierarchy = GEditor->AccessEditorUserSettings().bKeepAttachHierarchy;
+					const bool bKeepHierarchy = GetDefault<UEditorPerProjectUserSettings>()->bKeepAttachHierarchy;
 					Exporter->SetKeepHierarchy(bKeepHierarchy);
 
 					// Export the anim sequences
@@ -1979,7 +1979,7 @@ void FMatinee::OnContextGroupExportAnimFBX()
 							// Export the Matinee information to an FBX document.
 							Exporter->CreateDocument();
 							Exporter->SetTrasformBaking(bBakeTransforms);
-							const bool bKeepHierarchy = GEditor->AccessEditorUserSettings().bKeepAttachHierarchy;
+							const bool bKeepHierarchy = GetDefault<UEditorPerProjectUserSettings>()->bKeepAttachHierarchy;
 							Exporter->SetKeepHierarchy(bKeepHierarchy);
 
 							// Export the animation sequences in the group by sampling the skeletal mesh over the
@@ -3012,7 +3012,7 @@ void FMatinee::OnContextMoveMarkerToEndOfSelectedTrack()
 void FMatinee::OnToggleKeyframeBarSelection()
 {
 	bAllowKeyframeBarSelection = !bAllowKeyframeBarSelection;
-	GConfig->SetBool( TEXT("Matinee"), TEXT("AllowKeyframeBarSelection"), bAllowKeyframeBarSelection, GEditorUserSettingsIni );
+	GConfig->SetBool( TEXT("Matinee"), TEXT("AllowKeyframeBarSelection"), bAllowKeyframeBarSelection, GEditorPerProjectIni );
 }
 
 /**
@@ -3029,7 +3029,7 @@ bool FMatinee::IsKeyframeBarSelectionToggled()
 void FMatinee::OnToggleKeyframeTextSelection()
 {
 	bAllowKeyframeTextSelection = !bAllowKeyframeTextSelection;
-	GConfig->SetBool( TEXT("Matinee"), TEXT("AllowKeyframeTextSelection"), bAllowKeyframeTextSelection, GEditorUserSettingsIni );
+	GConfig->SetBool( TEXT("Matinee"), TEXT("AllowKeyframeTextSelection"), bAllowKeyframeTextSelection, GEditorPerProjectIni );
 }
 
 bool FMatinee::IsKeyframeTextSelectionToggled()
@@ -3851,7 +3851,7 @@ void FMatinee::OnMenuExport()
 				// Export the Matinee information to an FBX file
 				Exporter->CreateDocument();
 				Exporter->SetTrasformBaking(bBakeTransforms);
-				const bool bKeepHierarchy = GEditor->AccessEditorUserSettings().bKeepAttachHierarchy;
+				const bool bKeepHierarchy = GetDefault<UEditorPerProjectUserSettings>()->bKeepAttachHierarchy;
 				Exporter->SetKeepHierarchy(bKeepHierarchy);
 
 				const bool bSelectedOnly = false;
@@ -4236,7 +4236,8 @@ bool FMatinee::IsBakeTransformsToggled()
  */
 void FMatinee::OnToggleKeepHierarchy()
 {
-	GEditor->AccessEditorUserSettings().bKeepAttachHierarchy = !GEditor->AccessEditorUserSettings().bKeepAttachHierarchy;
+	auto* Settings = GetMutableDefault<UEditorPerProjectUserSettings>();
+	Settings->bKeepAttachHierarchy = !Settings->bKeepAttachHierarchy;
 }
 
 /**
@@ -4244,7 +4245,7 @@ void FMatinee::OnToggleKeepHierarchy()
  */
 bool FMatinee::IsKeepHierarchyToggled()
 {
-	return GEditor->AccessEditorUserSettings().bKeepAttachHierarchy;
+	return GetDefault<UEditorPerProjectUserSettings>()->bKeepAttachHierarchy;
 }
 
 
@@ -4258,7 +4259,7 @@ void FMatinee::OnToggleInvertPan()
 {
 	bInvertPan = !bInvertPan;
 
-	GConfig->SetBool(TEXT("Matinee"), TEXT("InterpEdPanInvert"), bInvertPan, GEditorUserSettingsIni);
+	GConfig->SetBool(TEXT("Matinee"), TEXT("InterpEdPanInvert"), bInvertPan, GEditorPerProjectIni);
 }
 
 bool FMatinee::IsInvertPanToggled()
@@ -4932,7 +4933,7 @@ void FMatinee::OnChangeSnapSize( TSharedPtr< FString > SelectedString, ESelectIn
 	SnapSelectionIndex = NewSelection;
 
 	// Save selected snap mode to INI.
-	GConfig->SetInt(TEXT("Matinee"), TEXT("SelectedSnapMode"), NewSelection, GEditorUserSettingsIni );
+	GConfig->SetInt(TEXT("Matinee"), TEXT("SelectedSnapMode"), NewSelection, GEditorPerProjectIni );
 
 	// Snap time to frames right now if we need to
 	SetSnapTimeToFrames( bSnapTimeToFrames );
@@ -4965,7 +4966,7 @@ void FMatinee::OnChangeInitialInterpMode( TSharedPtr<FString> ItemSelected, ESel
 	InitialInterpMode = ( EInterpCurveMode )NewSelection;
 
 	// Save selected mode to user's preference file
-	GConfig->SetInt( TEXT( "Matinee" ), TEXT( "InitialInterpMode2" ), NewSelection, GEditorUserSettingsIni );
+	GConfig->SetInt( TEXT( "Matinee" ), TEXT( "InitialInterpMode2" ), NewSelection, GEditorPerProjectIni );
 }
 
 
diff --git a/Engine/Source/Editor/Matinee/Private/MatineeRecordMovie.cpp b/Engine/Source/Editor/Matinee/Private/MatineeRecordMovie.cpp
--- a/Engine/Source/Editor/Matinee/Private/MatineeRecordMovie.cpp
+++ b/Engine/Source/Editor/Matinee/Private/MatineeRecordMovie.cpp
@@ -195,20 +195,20 @@ FReply SMatineeRecordMovie::OnOK()
 	Options.CustomRes = FString::Printf( TEXT("%dx%d"), CaptureWidth, CaptureHeight );
 	Options.DisableTextureStreaming = bDisableTextureStreaming;
 
-	GConfig->SetBool( TEXT("MatineeCreateMovieOptions"), TEXT("CloseEditor"), Options.CloseEditor, GEditorUserSettingsIni );
-	GConfig->SetBool( TEXT("MatineeCreateMovieOptions"), TEXT("Compress"), Options.Compress, GEditorUserSettingsIni );
-	GConfig->SetInt(  TEXT("MatineeCreateMovieOptions"), TEXT("CaptureResolutionFPS"), Options.CaptureResolutionFPS, GEditorUserSettingsIni );
-	GConfig->SetInt(  TEXT("MatineeCreateMovieOptions"), TEXT("CaptureResolutionIndex"), Options.CaptureResolutionIndex, GEditorUserSettingsIni );
-	GConfig->SetInt(  TEXT("MatineeCreateMovieOptions"), TEXT("CaptureTypeIndex"), Options.CaptureTypeIndex, GEditorUserSettingsIni );
-	GConfig->SetBool( TEXT("MatineeCreateMovieOptions"), TEXT("CinematicMode"), Options.CinematicMode, GEditorUserSettingsIni );
-	GConfig->SetBool( TEXT("MatineeCreateMovieOptions"), TEXT("DisableMovement"), Options.DisableMovement, GEditorUserSettingsIni );
-	GConfig->SetBool( TEXT("MatineeCreateMovieOptions"), TEXT("DisableTurning"), Options.DisableTurning, GEditorUserSettingsIni );
-	GConfig->SetBool( TEXT("MatineeCreateMovieOptions"), TEXT("HidePlayer"), Options.HidePlayer, GEditorUserSettingsIni );
-	GConfig->SetBool( TEXT("MatineeCreateMovieOptions"), TEXT("HideHUD"), Options.HideHUD, GEditorUserSettingsIni );
-	GConfig->SetBool( TEXT("MatineeCreateMovieOptions"), TEXT("DisableTextureStreaming"), Options.DisableTextureStreaming, GEditorUserSettingsIni );
-	GConfig->SetString( TEXT("MatineeCreateMovieOptions"), TEXT("CustomRes"), *Options.CustomRes, GEditorUserSettingsIni );
-
-	GConfig->Flush( false, GEditorUserSettingsIni );
+	GConfig->SetBool( TEXT("MatineeCreateMovieOptions"), TEXT("CloseEditor"), Options.CloseEditor, GEditorPerProjectIni );
+	GConfig->SetBool( TEXT("MatineeCreateMovieOptions"), TEXT("Compress"), Options.Compress, GEditorPerProjectIni );
+	GConfig->SetInt(  TEXT("MatineeCreateMovieOptions"), TEXT("CaptureResolutionFPS"), Options.CaptureResolutionFPS, GEditorPerProjectIni );
+	GConfig->SetInt(  TEXT("MatineeCreateMovieOptions"), TEXT("CaptureResolutionIndex"), Options.CaptureResolutionIndex, GEditorPerProjectIni );
+	GConfig->SetInt(  TEXT("MatineeCreateMovieOptions"), TEXT("CaptureTypeIndex"), Options.CaptureTypeIndex, GEditorPerProjectIni );
+	GConfig->SetBool( TEXT("MatineeCreateMovieOptions"), TEXT("CinematicMode"), Options.CinematicMode, GEditorPerProjectIni );
+	GConfig->SetBool( TEXT("MatineeCreateMovieOptions"), TEXT("DisableMovement"), Options.DisableMovement, GEditorPerProjectIni );
+	GConfig->SetBool( TEXT("MatineeCreateMovieOptions"), TEXT("DisableTurning"), Options.DisableTurning, GEditorPerProjectIni );
+	GConfig->SetBool( TEXT("MatineeCreateMovieOptions"), TEXT("HidePlayer"), Options.HidePlayer, GEditorPerProjectIni );
+	GConfig->SetBool( TEXT("MatineeCreateMovieOptions"), TEXT("HideHUD"), Options.HideHUD, GEditorPerProjectIni );
+	GConfig->SetBool( TEXT("MatineeCreateMovieOptions"), TEXT("DisableTextureStreaming"), Options.DisableTextureStreaming, GEditorPerProjectIni );
+	GConfig->SetString( TEXT("MatineeCreateMovieOptions"), TEXT("CustomRes"), *Options.CustomRes, GEditorPerProjectIni );
+
+	GConfig->Flush( false, GEditorPerProjectIni );
 
 	FEdModeInterpEdit* Mode = (FEdModeInterpEdit*)GLevelEditorModeTools().GetActiveMode( FBuiltinEditorModes::EM_InterpEdit );
 	if ( Mode != NULL && Mode->InterpEd != NULL )
@@ -246,19 +246,19 @@ FReply SMatineeRecordMovie::OnCancel()
 void SMatineeRecordMovie::InitializeOptions()
 {
 	// Retrieve settings from ini.
-	GConfig->GetBool( TEXT("MatineeCreateMovieOptions"), TEXT("CloseEditor"), Options.CloseEditor, GEditorUserSettingsIni );
-	GConfig->GetBool( TEXT("MatineeCreateMovieOptions"), TEXT("Compress"), Options.Compress, GEditorUserSettingsIni );
-	GConfig->GetInt(  TEXT("MatineeCreateMovieOptions"), TEXT("CaptureResolutionFPS"), Options.CaptureResolutionFPS, GEditorUserSettingsIni );
-	GConfig->GetInt(  TEXT("MatineeCreateMovieOptions"), TEXT("CaptureResolutionIndex"), Options.CaptureResolutionIndex, GEditorUserSettingsIni );
-	GConfig->GetInt(  TEXT("MatineeCreateMovieOptions"), TEXT("CaptureTypeIndex"), Options.CaptureTypeIndex, GEditorUserSettingsIni );
-	GConfig->GetBool( TEXT("MatineeCreateMovieOptions"), TEXT("CinematicMode"), Options.CinematicMode, GEditorUserSettingsIni );
-	GConfig->GetBool( TEXT("MatineeCreateMovieOptions"), TEXT("DisableMovement"), Options.DisableMovement, GEditorUserSettingsIni );
-	GConfig->GetBool( TEXT("MatineeCreateMovieOptions"), TEXT("DisableTurning"), Options.DisableTurning, GEditorUserSettingsIni );
-	GConfig->GetBool( TEXT("MatineeCreateMovieOptions"), TEXT("HidePlayer"), Options.HidePlayer, GEditorUserSettingsIni );
-	GConfig->GetBool( TEXT("MatineeCreateMovieOptions"), TEXT("HideHUD"), Options.HideHUD, GEditorUserSettingsIni );
-	GConfig->GetBool( TEXT("MatineeCreateMovieOptions"), TEXT("DisableTextureStreaming"), Options.DisableTextureStreaming, GEditorUserSettingsIni );
-
-	GConfig->GetString( TEXT("MatineeCreateMovieOptions"), TEXT("CustomRes"), Options.CustomRes, GEditorUserSettingsIni );
+	GConfig->GetBool( TEXT("MatineeCreateMovieOptions"), TEXT("CloseEditor"), Options.CloseEditor, GEditorPerProjectIni );
+	GConfig->GetBool( TEXT("MatineeCreateMovieOptions"), TEXT("Compress"), Options.Compress, GEditorPerProjectIni );
+	GConfig->GetInt(  TEXT("MatineeCreateMovieOptions"), TEXT("CaptureResolutionFPS"), Options.CaptureResolutionFPS, GEditorPerProjectIni );
+	GConfig->GetInt(  TEXT("MatineeCreateMovieOptions"), TEXT("CaptureResolutionIndex"), Options.CaptureResolutionIndex, GEditorPerProjectIni );
+	GConfig->GetInt(  TEXT("MatineeCreateMovieOptions"), TEXT("CaptureTypeIndex"), Options.CaptureTypeIndex, GEditorPerProjectIni );
+	GConfig->GetBool( TEXT("MatineeCreateMovieOptions"), TEXT("CinematicMode"), Options.CinematicMode, GEditorPerProjectIni );
+	GConfig->GetBool( TEXT("MatineeCreateMovieOptions"), TEXT("DisableMovement"), Options.DisableMovement, GEditorPerProjectIni );
+	GConfig->GetBool( TEXT("MatineeCreateMovieOptions"), TEXT("DisableTurning"), Options.DisableTurning, GEditorPerProjectIni );
+	GConfig->GetBool( TEXT("MatineeCreateMovieOptions"), TEXT("HidePlayer"), Options.HidePlayer, GEditorPerProjectIni );
+	GConfig->GetBool( TEXT("MatineeCreateMovieOptions"), TEXT("HideHUD"), Options.HideHUD, GEditorPerProjectIni );
+	GConfig->GetBool( TEXT("MatineeCreateMovieOptions"), TEXT("DisableTextureStreaming"), Options.DisableTextureStreaming, GEditorPerProjectIni );
+
+	GConfig->GetString( TEXT("MatineeCreateMovieOptions"), TEXT("CustomRes"), Options.CustomRes, GEditorPerProjectIni );
 
 	// setup format strings
 	CaptureFormats.Empty();
diff --git a/Engine/Source/Editor/Matinee/Private/MatineeTools.cpp b/Engine/Source/Editor/Matinee/Private/MatineeTools.cpp
--- a/Engine/Source/Editor/Matinee/Private/MatineeTools.cpp
+++ b/Engine/Source/Editor/Matinee/Private/MatineeTools.cpp
@@ -1600,7 +1600,7 @@ void FMatinee::SetSnapEnabled(bool bInSnapEnabled)
 	}
 
 	// Save to ini when it changes.
-	GConfig->SetBool( TEXT("Matinee"), TEXT("SnapEnabled"), bSnapEnabled, GEditorUserSettingsIni );
+	GConfig->SetBool( TEXT("Matinee"), TEXT("SnapEnabled"), bSnapEnabled, GEditorPerProjectIni );
 }
 
 
@@ -1611,7 +1611,7 @@ void FMatinee::SetSnapTimeToFrames( bool bInValue )
 	bSnapTimeToFrames = bInValue;
 
 	// Save to ini when it changes.
-	GConfig->SetBool( TEXT("Matinee"), TEXT("SnapTimeToFrames"), bSnapTimeToFrames, GEditorUserSettingsIni );
+	GConfig->SetBool( TEXT("Matinee"), TEXT("SnapTimeToFrames"), bSnapTimeToFrames, GEditorPerProjectIni );
 
 	// Go ahead and apply the change right now if we need to
 	if( IsInitialized() && bSnapToFrames && bSnapTimeToFrames )
@@ -1628,7 +1628,7 @@ void FMatinee::SetFixedTimeStepPlayback( bool bInValue )
 	bFixedTimeStepPlayback = bInValue;
 	
 	// Save to ini when it changes.
-	GConfig->SetBool( TEXT("Matinee"), TEXT("FixedTimeStepPlayback"), bFixedTimeStepPlayback, GEditorUserSettingsIni );
+	GConfig->SetBool( TEXT("Matinee"), TEXT("FixedTimeStepPlayback"), bFixedTimeStepPlayback, GEditorPerProjectIni );
 
 	// Update fixed time step state
 	UpdateFixedTimeStepPlayback();
@@ -1655,7 +1655,7 @@ void FMatinee::SetPreferFrameNumbers( bool bInValue )
 	bPreferFrameNumbers = bInValue;
 
 	// Save to ini when it changes.
-	GConfig->SetBool( TEXT("Matinee"), TEXT("PreferFrameNumbers"), bPreferFrameNumbers, GEditorUserSettingsIni );
+	GConfig->SetBool( TEXT("Matinee"), TEXT("PreferFrameNumbers"), bPreferFrameNumbers, GEditorPerProjectIni );
 }
 
 
@@ -1666,7 +1666,7 @@ void FMatinee::SetShowTimeCursorPosForAllKeys( bool bInValue )
 	bShowTimeCursorPosForAllKeys = bInValue;
 
 	// Save to ini when it changes.
-	GConfig->SetBool( TEXT("Matinee"), TEXT("ShowTimeCursorPosForAllKeys"), bShowTimeCursorPosForAllKeys, GEditorUserSettingsIni );
+	GConfig->SetBool( TEXT("Matinee"), TEXT("ShowTimeCursorPosForAllKeys"), bShowTimeCursorPosForAllKeys, GEditorPerProjectIni );
 }
 
 
diff --git a/Engine/Source/Editor/Matinee/Private/MatineeTrackView.cpp b/Engine/Source/Editor/Matinee/Private/MatineeTrackView.cpp
--- a/Engine/Source/Editor/Matinee/Private/MatineeTrackView.cpp
+++ b/Engine/Source/Editor/Matinee/Private/MatineeTrackView.cpp
@@ -148,7 +148,7 @@ void FMatinee::OnToggleDirectorTimeline()
 		DirectorTrackWindow->InterpEdVC->Viewport->Draw();
 
 		// Save to ini when it changes.
-		GConfig->SetBool(TEXT("Matinee"), TEXT("DirectorTimelineEnabled"), DirectorTrackWindow->InterpEdVC->bWantTimeline, GEditorUserSettingsIni);
+		GConfig->SetBool(TEXT("Matinee"), TEXT("DirectorTimelineEnabled"), DirectorTrackWindow->InterpEdVC->bWantTimeline, GEditorPerProjectIni);
 	}
 }
 
diff --git a/Engine/Source/Editor/MeshPaint/Private/MeshPaintEdMode.cpp b/Engine/Source/Editor/MeshPaint/Private/MeshPaintEdMode.cpp
--- a/Engine/Source/Editor/MeshPaint/Private/MeshPaintEdMode.cpp
+++ b/Engine/Source/Editor/MeshPaint/Private/MeshPaintEdMode.cpp
@@ -4076,7 +4076,7 @@ void FEdModeMeshPaint::SetBrushRadiiDefault( float InBrushRadius )
 	GetBrushRadiiLimits(MinBrushRadius, MaxBrushRadius);	
 
 	InBrushRadius = (float)FMath::Clamp(InBrushRadius, MinBrushRadius, MaxBrushRadius);
-	GConfig->SetFloat( TEXT("MeshPaintEdit"), TEXT("DefaultBrushRadius"), InBrushRadius, GEditorUserSettingsIni );
+	GConfig->SetFloat( TEXT("MeshPaintEdit"), TEXT("DefaultBrushRadius"), InBrushRadius, GEditorPerProjectIni );
 }
 
 float FEdModeMeshPaint::GetBrushRadiiDefault() const
@@ -4085,7 +4085,7 @@ float FEdModeMeshPaint::GetBrushRadiiDefault() const
 	GetBrushRadiiLimits(MinBrushRadius, MaxBrushRadius);
 
 	float BrushRadius = 128.f;
-	GConfig->GetFloat( TEXT("MeshPaintEdit"), TEXT("DefaultBrushRadius"), BrushRadius, GEditorUserSettingsIni );
+	GConfig->GetFloat( TEXT("MeshPaintEdit"), TEXT("DefaultBrushRadius"), BrushRadius, GEditorPerProjectIni );
 	BrushRadius = (float)FMath::Clamp(BrushRadius, MinBrushRadius, MaxBrushRadius);
 	return BrushRadius;
 }
diff --git a/Engine/Source/Editor/NewsFeed/Private/Implementation/NewsFeedCache.cpp b/Engine/Source/Editor/NewsFeed/Private/Implementation/NewsFeedCache.cpp
--- a/Engine/Source/Editor/NewsFeed/Private/Implementation/NewsFeedCache.cpp
+++ b/Engine/Source/Editor/NewsFeed/Private/Implementation/NewsFeedCache.cpp
@@ -307,11 +307,7 @@ void FNewsFeedCache::SaveSettings( )
 		}
 	}
 
-	// Save config to file, but only if we are not the build machine since settings may put the builder in an unclean state
-	if (!GIsBuildMachine)
-	{
-		NewsFeedSettings->SaveConfig();
-	}
+	NewsFeedSettings->SaveConfig();
 }
 
 
diff --git a/Engine/Source/Editor/PlacementMode/Private/PlacementMode.cpp b/Engine/Source/Editor/PlacementMode/Private/PlacementMode.cpp
--- a/Engine/Source/Editor/PlacementMode/Private/PlacementMode.cpp
+++ b/Engine/Source/Editor/PlacementMode/Private/PlacementMode.cpp
@@ -29,8 +29,8 @@ FPlacementMode::~FPlacementMode()
 
 void FPlacementMode::Initialize()
 {
-	//GConfig->GetFloat(TEXT("PlacementMode"), TEXT("AssetThumbnailScale"), AssetThumbnailScale, GEditorUserSettingsIni);
-	//GConfig->GetBool( TEXT( "PlacementMode" ), TEXT( "ShowOtherDeveloperAssets" ), ShowOtherDeveloperAssets, GEditorUserSettingsIni );
+	//GConfig->GetFloat(TEXT("PlacementMode"), TEXT("AssetThumbnailScale"), AssetThumbnailScale, GEditorPerProjectIni);
+	//GConfig->GetBool( TEXT( "PlacementMode" ), TEXT( "ShowOtherDeveloperAssets" ), ShowOtherDeveloperAssets, GEditorPerProjectIni );
 }
 
 bool FPlacementMode::UsesToolkits() const
diff --git a/Engine/Source/Editor/PlacementMode/Private/PlacementModeModule.cpp b/Engine/Source/Editor/PlacementMode/Private/PlacementModeModule.cpp
--- a/Engine/Source/Editor/PlacementMode/Private/PlacementModeModule.cpp
+++ b/Engine/Source/Editor/PlacementMode/Private/PlacementModeModule.cpp
@@ -15,10 +15,10 @@ public:
 	{
 
 		TArray< FString > RecentlyPlacedAsStrings;
-		GConfig->GetArray(TEXT("PlacementMode"), TEXT("RecentlyPlaced"), RecentlyPlacedAsStrings, GEditorUserSettingsIni);
+		GConfig->GetArray(TEXT("PlacementMode"), TEXT("RecentlyPlaced"), RecentlyPlacedAsStrings, GEditorPerProjectIni);
 
 		//FString ActivePaletteName;
-		//GConfig->GetString( TEXT( "PlacementMode" ), TEXT( "ActivePalette" ), ActivePaletteName, GEditorUserSettingsIni );
+		//GConfig->GetString( TEXT( "PlacementMode" ), TEXT( "ActivePalette" ), ActivePaletteName, GEditorPerProjectIni );
 
 		for (int Index = 0; Index < RecentlyPlacedAsStrings.Num(); Index++)
 		{
@@ -116,7 +116,7 @@ public:
 				RecentlyPlacedAsStrings.Add(RecentlyPlaced[Index].ToString());
 			}
 
-			GConfig->SetArray(TEXT("PlacementMode"), TEXT("RecentlyPlaced"), RecentlyPlacedAsStrings, GEditorUserSettingsIni);
+			GConfig->SetArray(TEXT("PlacementMode"), TEXT("RecentlyPlaced"), RecentlyPlacedAsStrings, GEditorPerProjectIni);
 			RecentlyPlacedChanged.Broadcast(RecentlyPlaced);
 		}
 	}
diff --git a/Engine/Source/Editor/PropertyEditor/Private/DetailCategoryBuilderImpl.cpp b/Engine/Source/Editor/PropertyEditor/Private/DetailCategoryBuilderImpl.cpp
--- a/Engine/Source/Editor/PropertyEditor/Private/DetailCategoryBuilderImpl.cpp
+++ b/Engine/Source/Editor/PropertyEditor/Private/DetailCategoryBuilderImpl.cpp
@@ -93,7 +93,7 @@ FDetailCategoryImpl::FDetailCategoryImpl( FName InCategoryName, TSharedRef<FDeta
 
 	CategoryPathName = BaseStructName.ToString() + TEXT(".") + CategoryName.ToString();
 
-	GConfig->GetBool( TEXT("DetailCategoriesAdvanced"), *CategoryPathName, bUserShowAdvanced, GEditorUserSettingsIni );
+	GConfig->GetBool( TEXT("DetailCategoriesAdvanced"), *CategoryPathName, bUserShowAdvanced, GEditorPerProjectIni );
 
 }
 
@@ -402,7 +402,7 @@ void FDetailCategoryImpl::OnAdvancedDropdownClicked()
 {
 	bUserShowAdvanced = !bUserShowAdvanced;
 
-	GConfig->SetBool( TEXT("DetailCategoriesAdvanced"), *CategoryPathName, bUserShowAdvanced, GEditorUserSettingsIni );
+	GConfig->SetBool( TEXT("DetailCategoriesAdvanced"), *CategoryPathName, bUserShowAdvanced, GEditorPerProjectIni );
 
 	const bool bRefilterCategory = true;
 	RefreshTree( bRefilterCategory );
@@ -546,7 +546,7 @@ void FDetailCategoryImpl::OnItemExpansionChanged( bool bIsExpanded )
 	if( bRestoreExpansionState )
 	{
 		// Save the collapsed state of this section
-		GConfig->SetBool( TEXT("DetailCategories"), *CategoryPathName, bIsExpanded, GEditorUserSettingsIni );
+		GConfig->SetBool( TEXT("DetailCategories"), *CategoryPathName, bIsExpanded, GEditorPerProjectIni );
 	}
 
 	OnExpansionChangedDelegate.ExecuteIfBound( bIsExpanded );
@@ -563,7 +563,7 @@ bool FDetailCategoryImpl::ShouldBeExpanded() const
 		// Collapse by default if there are no simple child nodes
 		bool bShouldBeExpanded = !ContainsOnlyAdvanced();
 		// Save the collapsed state of this section
-		GConfig->GetBool( TEXT("DetailCategories"), *CategoryPathName, bShouldBeExpanded, GEditorUserSettingsIni );
+		GConfig->GetBool( TEXT("DetailCategories"), *CategoryPathName, bShouldBeExpanded, GEditorPerProjectIni );
 		return bShouldBeExpanded;
 	}
 	else
diff --git a/Engine/Source/Editor/PropertyEditor/Private/ObjectPropertyNode.cpp b/Engine/Source/Editor/PropertyEditor/Private/ObjectPropertyNode.cpp
--- a/Engine/Source/Editor/PropertyEditor/Private/ObjectPropertyNode.cpp
+++ b/Engine/Source/Editor/PropertyEditor/Private/ObjectPropertyNode.cpp
@@ -340,7 +340,7 @@ void FObjectPropertyNode::InternalInitChildNodes( FName SinglePropertyName )
 		if (MetaDataVisibilityCheckString.Len())
 		{
 			//ensure that the metadata visibility string is actually set to true in order to show this property
-			GConfig->GetBool(TEXT("UnrealEd.PropertyFilters"), *MetaDataVisibilityCheckString, bMetaDataAllowVisible, GEditorUserSettingsIni);
+			GConfig->GetBool(TEXT("UnrealEd.PropertyFilters"), *MetaDataVisibilityCheckString, bMetaDataAllowVisible, GEditorPerProjectIni);
 		}
 
 		if (bMetaDataAllowVisible)
diff --git a/Engine/Source/Editor/PropertyEditor/Private/Presentation/PropertyTable/PropertyTable.cpp b/Engine/Source/Editor/PropertyEditor/Private/Presentation/PropertyTable/PropertyTable.cpp
--- a/Engine/Source/Editor/PropertyEditor/Private/Presentation/PropertyTable/PropertyTable.cpp
+++ b/Engine/Source/Editor/PropertyEditor/Private/Presentation/PropertyTable/PropertyTable.cpp
@@ -419,7 +419,7 @@ void FPropertyTable::PasteTextAtCell( const FString& Text, const TSharedRef< IPr
 	RowStrings[CurrentRowIdx++].ParseIntoArray(CellStrings, TEXT("\t"), false);
 
 	// Get the maximum paste operations before displaying the slow task
-	int32 NumPasteOperationsBeforeWarning = GEditor->AccessEditorUserSettings().PropertyMatrix_NumberOfPasteOperationsBeforeWarning;
+	int32 NumPasteOperationsBeforeWarning = GetDefault<UEditorPerProjectUserSettings>()->PropertyMatrix_NumberOfPasteOperationsBeforeWarning;
 	
 	const bool bShowCancelButton = false;
 	const bool bShowProgressDialog = SelectedCells.Num() > NumPasteOperationsBeforeWarning;
diff --git a/Engine/Source/Editor/PropertyEditor/Private/PropertyNode.cpp b/Engine/Source/Editor/PropertyEditor/Private/PropertyNode.cpp
--- a/Engine/Source/Editor/PropertyEditor/Private/PropertyNode.cpp
+++ b/Engine/Source/Editor/PropertyEditor/Private/PropertyNode.cpp
@@ -22,9 +22,9 @@ FPropertySettings::FPropertySettings()
 	, bExpandDistributions( false )
 	, bShowHiddenProperties(false)
 {
-	GConfig->GetBool(TEXT("PropertySettings"), TEXT("ShowHiddenProperties"), bShowHiddenProperties, GEditorUserSettingsIni);
-	GConfig->GetBool(TEXT("PropertySettings"), TEXT("ShowFriendlyPropertyNames"), bShowFriendlyPropertyNames, GEditorUserSettingsIni);
-	GConfig->GetBool(TEXT("PropertySettings"), TEXT("ExpandDistributions"), bExpandDistributions, GEditorUserSettingsIni);
+	GConfig->GetBool(TEXT("PropertySettings"), TEXT("ShowHiddenProperties"), bShowHiddenProperties, GEditorPerProjectIni);
+	GConfig->GetBool(TEXT("PropertySettings"), TEXT("ShowFriendlyPropertyNames"), bShowFriendlyPropertyNames, GEditorPerProjectIni);
+	GConfig->GetBool(TEXT("PropertySettings"), TEXT("ExpandDistributions"), bExpandDistributions, GEditorPerProjectIni);
 }
 
 DEFINE_LOG_CATEGORY(LogPropertyNode);
diff --git a/Engine/Source/Editor/PropertyEditor/Private/SDetailsViewBase.cpp b/Engine/Source/Editor/PropertyEditor/Private/SDetailsViewBase.cpp
--- a/Engine/Source/Editor/PropertyEditor/Private/SDetailsViewBase.cpp
+++ b/Engine/Source/Editor/PropertyEditor/Private/SDetailsViewBase.cpp
@@ -884,13 +884,13 @@ void SDetailsViewBase::SaveExpandedItems()
 			if (!bShouldSave)
 			{
 				TArray<FString> DummyExpandedPropertyItems;
-				GConfig->GetSingleLineArray(TEXT("DetailPropertyExpansion"), *Struct->GetName(), DummyExpandedPropertyItems, GEditorUserSettingsIni);
+				GConfig->GetSingleLineArray(TEXT("DetailPropertyExpansion"), *Struct->GetName(), DummyExpandedPropertyItems, GEditorPerProjectIni);
 				bShouldSave = DummyExpandedPropertyItems.Num() > 0;
 			}
 
 			if (bShouldSave)
 			{
-				GConfig->SetSingleLineArray(TEXT("DetailPropertyExpansion"), *Struct->GetName(), ExpandedPropertyItems, GEditorUserSettingsIni);
+				GConfig->SetSingleLineArray(TEXT("DetailPropertyExpansion"), *Struct->GetName(), ExpandedPropertyItems, GEditorPerProjectIni);
 			}
 		}
 	}
@@ -901,12 +901,12 @@ void SDetailsViewBase::SaveExpandedItems()
 		if (!bShouldSave)
 		{
 			FString DummyExpandedCustomItemsString;
-			GConfig->GetString(TEXT("DetailCustomWidgetExpansion"), *BestBaseStruct->GetName(), DummyExpandedCustomItemsString, GEditorUserSettingsIni);
+			GConfig->GetString(TEXT("DetailCustomWidgetExpansion"), *BestBaseStruct->GetName(), DummyExpandedCustomItemsString, GEditorPerProjectIni);
 			bShouldSave = !DummyExpandedCustomItemsString.IsEmpty();
 		}
 		if (bShouldSave)
 		{
-			GConfig->SetString(TEXT("DetailCustomWidgetExpansion"), *BestBaseStruct->GetName(), *ExpandedCustomItemsString, GEditorUserSettingsIni);
+			GConfig->SetString(TEXT("DetailCustomWidgetExpansion"), *BestBaseStruct->GetName(), *ExpandedCustomItemsString, GEditorPerProjectIni);
 		}
 	}
 }
@@ -929,13 +929,13 @@ void SDetailsViewBase::RestoreExpandedItems(TSharedPtr<FPropertyNode> InitialSta
 	//while a valid class, and we're either the same as the base class (for multiple actors being selected and base class is AActor) OR we're not down to AActor yet)
 	for (UStruct* Struct = BestBaseStruct; Struct && ((BestBaseStruct == Struct) || (Struct != AActor::StaticClass())); Struct = Struct->GetSuperStruct())
 	{
-		GConfig->GetSingleLineArray(TEXT("DetailPropertyExpansion"), *Struct->GetName(), ExpandedPropertyItems, GEditorUserSettingsIni);
+		GConfig->GetSingleLineArray(TEXT("DetailPropertyExpansion"), *Struct->GetName(), ExpandedPropertyItems, GEditorPerProjectIni);
 		SetExpandedItems(StartNode, ExpandedPropertyItems);
 	}
 
 	if (BestBaseStruct)
 	{
-		GConfig->GetString(TEXT("DetailCustomWidgetExpansion"), *BestBaseStruct->GetName(), ExpandedCustomItems, GEditorUserSettingsIni);
+		GConfig->GetString(TEXT("DetailCustomWidgetExpansion"), *BestBaseStruct->GetName(), ExpandedCustomItems, GEditorPerProjectIni);
 		TArray<FString> ExpandedCustomItemsArray;
 		ExpandedCustomItems.ParseIntoArray(ExpandedCustomItemsArray, TEXT(","), true);
 
diff --git a/Engine/Source/Editor/PropertyEditor/Private/SPropertyTreeViewImpl.cpp b/Engine/Source/Editor/PropertyEditor/Private/SPropertyTreeViewImpl.cpp
--- a/Engine/Source/Editor/PropertyEditor/Private/SPropertyTreeViewImpl.cpp
+++ b/Engine/Source/Editor/PropertyEditor/Private/SPropertyTreeViewImpl.cpp
@@ -112,7 +112,7 @@ void SPropertyTreeViewImpl::Construct(const FArguments& InArgs)
 	ConstructExternalColumnHeaders = InArgs._ConstructExternalColumnHeaders;
 	ConstructExternalColumnCell = InArgs._ConstructExternalColumnCell;
 
-	if( !GConfig->GetBool(TEXT("PropertyWindow"), TEXT("ShowFavoritesWindow"), bFavoritesEnabled, GEditorUserSettingsIni) )
+	if( !GConfig->GetBool(TEXT("PropertyWindow"), TEXT("ShowFavoritesWindow"), bFavoritesEnabled, GEditorPerProjectIni) )
 	{
 		bFavoritesEnabled = false;
 	}
@@ -278,7 +278,7 @@ FReply SPropertyTreeViewImpl::OnToggleFavoritesClicked()
 	bFavoritesEnabled = !bFavoritesEnabled;
 
 	//save off state of the filter window
-	GConfig->SetBool(TEXT("PropertyWindow"), TEXT("ShowFavoritesWindow"), bFavoritesEnabled, GEditorUserSettingsIni);
+	GConfig->SetBool(TEXT("PropertyWindow"), TEXT("ShowFavoritesWindow"), bFavoritesEnabled, GEditorPerProjectIni);
 
 	return FReply::Handled();
 }
@@ -449,7 +449,7 @@ void SPropertyTreeViewImpl::SaveExpandedItems()
 // 				ExpansionName += TEXT("Favorites");
 // 			}
 
-			GConfig->SetSingleLineArray(TEXT("PropertyWindowExpansion"), *ExpansionName, ExpandedItemNames, GEditorUserSettingsIni);
+			GConfig->SetSingleLineArray(TEXT("PropertyWindowExpansion"), *ExpansionName, ExpandedItemNames, GEditorPerProjectIni);
 		}
 	}
 }
@@ -461,7 +461,7 @@ void SPropertyTreeViewImpl::SaveColumnWidths()
 	{
 		const SHeaderRow::FColumn& Column = Columns[Idx];
 		const float Width = Column.GetWidth();
-		GConfig->SetFloat(TEXT("PropertyWindowWidths"), *Column.ColumnId.ToString(), Width, GEditorUserSettingsIni);
+		GConfig->SetFloat(TEXT("PropertyWindowWidths"), *Column.ColumnId.ToString(), Width, GEditorPerProjectIni);
 	}
 }
 
@@ -488,7 +488,7 @@ void SPropertyTreeViewImpl::RestoreExpandedItems()
 // 			ExpansionName += TEXT("Favorites");
 // 		}
 
-		GConfig->GetSingleLineArray(TEXT("PropertyWindowExpansion"), *ExpansionName, ExpandedItems, GEditorUserSettingsIni);
+		GConfig->GetSingleLineArray(TEXT("PropertyWindowExpansion"), *ExpansionName, ExpandedItems, GEditorPerProjectIni);
 		
 		TSharedRef<SPropertyTree> PropertyTreeRef = PropertyTree.ToSharedRef();
 		SetExpandedItems( RootPropertyNode, PropertyTreeRef, ExpandedItems );
@@ -503,7 +503,7 @@ void SPropertyTreeViewImpl::RestoreColumnWidths()
 	{
 		const SHeaderRow::FColumn& Column = Columns[Idx];
 		float Width = 1.0f;
-		if ( GConfig->GetFloat(TEXT("PropertyWindowWidths"), *Column.ColumnId.ToString(), Width, GEditorUserSettingsIni) )
+		if ( GConfig->GetFloat(TEXT("PropertyWindowWidths"), *Column.ColumnId.ToString(), Width, GEditorPerProjectIni) )
 		{
 			ColumnHeaderRow->SetColumnWidth( Column.ColumnId, Width );
 		}
@@ -1128,7 +1128,7 @@ void SPropertyTreeViewImpl::LoadFavorites()
 			FString ContextName = BestClass->GetName() + TEXT("Favorites");
 
 			TArray<FString> OutFavoritesList;
-			GConfig->GetSingleLineArray(TEXT("PropertyWindow"), *ContextName, OutFavoritesList, GEditorUserSettingsIni);
+			GConfig->GetSingleLineArray(TEXT("PropertyWindow"), *ContextName, OutFavoritesList, GEditorPerProjectIni);
 		
 			for(int32 i=0; i<OutFavoritesList.Num(); i++)
 			{
@@ -1162,7 +1162,7 @@ void SPropertyTreeViewImpl::SaveFavorites()
 				FavoritesArray.Add( *It );
 			}
 
-			GConfig->SetSingleLineArray(TEXT("PropertyWindow"), *ContextName, FavoritesArray, GEditorUserSettingsIni);
+			GConfig->SetSingleLineArray(TEXT("PropertyWindow"), *ContextName, FavoritesArray, GEditorPerProjectIni);
 		}
 	}
 }
diff --git a/Engine/Source/Editor/StatsViewer/Private/SStatsViewer.cpp b/Engine/Source/Editor/StatsViewer/Private/SStatsViewer.cpp
--- a/Engine/Source/Editor/StatsViewer/Private/SStatsViewer.cpp
+++ b/Engine/Source/Editor/StatsViewer/Private/SStatsViewer.cpp
@@ -199,7 +199,7 @@ void SStatsViewer::Construct( const FArguments& InArgs )
 	{
 		TSharedPtr<IStatsPage> InitialStatsPage;
 		FString DisplayedStatsPageName;
-		if (GConfig->GetString(*StatsViewerConstants::ConfigSectionName, TEXT("DisplayedStatsPageName"), DisplayedStatsPageName, GEditorUserSettingsIni))
+		if (GConfig->GetString(*StatsViewerConstants::ConfigSectionName, TEXT("DisplayedStatsPageName"), DisplayedStatsPageName, GEditorPerProjectIni))
 		{
 			InitialStatsPage = FStatsPageManager::Get().GetPage(FName(*DisplayedStatsPageName));
 		}
@@ -701,7 +701,7 @@ void SStatsViewer::SetDisplayedStats( TSharedRef<IStatsPage> StatsPage )
 	}
 
 	CurrentStats = StatsPage;
-	GConfig->SetString(*StatsViewerConstants::ConfigSectionName, TEXT("DisplayedStatsPageName"), *StatsPage->GetName().ToString(), GEditorUserSettingsIni);
+	GConfig->SetString(*StatsViewerConstants::ConfigSectionName, TEXT("DisplayedStatsPageName"), *StatsPage->GetName().ToString(), GEditorPerProjectIni);
 	CurrentStats->OnShow( SharedThis(this) );
 	CurrentObjectSetIndex = CurrentStats->GetSelectedObjectSet();
 	CurrentFilterIndex = 0;
diff --git a/Engine/Source/Editor/UMGEditor/Private/Designer/SDesignerView.cpp b/Engine/Source/Editor/UMGEditor/Private/Designer/SDesignerView.cpp
--- a/Engine/Source/Editor/UMGEditor/Private/Designer/SDesignerView.cpp
+++ b/Engine/Source/Editor/UMGEditor/Private/Designer/SDesignerView.cpp
@@ -613,21 +613,21 @@ void SDesignerView::SetStartupResolution()
 {
 	// Use previously set resolution (or create new entries using default values)
 	// Width
-	if (!GConfig->GetInt(*ConfigSectionName, TEXT("PreviewWidth"), PreviewWidth, GEditorUserSettingsIni))
+	if (!GConfig->GetInt(*ConfigSectionName, TEXT("PreviewWidth"), PreviewWidth, GEditorPerProjectIni))
 	{
-		GConfig->SetInt(*ConfigSectionName, TEXT("PreviewWidth"), DefaultResolutionWidth, GEditorUserSettingsIni);
+		GConfig->SetInt(*ConfigSectionName, TEXT("PreviewWidth"), DefaultResolutionWidth, GEditorPerProjectIni);
 		PreviewWidth = DefaultResolutionWidth;
 	}
 	// Height
-	if (!GConfig->GetInt(*ConfigSectionName, TEXT("PreviewHeight"), PreviewHeight, GEditorUserSettingsIni))
+	if (!GConfig->GetInt(*ConfigSectionName, TEXT("PreviewHeight"), PreviewHeight, GEditorPerProjectIni))
 	{
-		GConfig->SetInt(*ConfigSectionName, TEXT("PreviewHeight"), DefaultResolutionHeight, GEditorUserSettingsIni);
+		GConfig->SetInt(*ConfigSectionName, TEXT("PreviewHeight"), DefaultResolutionHeight, GEditorPerProjectIni);
 		PreviewHeight = DefaultResolutionHeight;
 	}
 	// Aspect Ratio
-	if (!GConfig->GetString(*ConfigSectionName, TEXT("PreviewAspectRatio"), PreviewAspectRatio, GEditorUserSettingsIni))
+	if (!GConfig->GetString(*ConfigSectionName, TEXT("PreviewAspectRatio"), PreviewAspectRatio, GEditorPerProjectIni))
 	{
-		GConfig->SetString(*ConfigSectionName, TEXT("PreviewAspectRatio"), *DefaultAspectRatio, GEditorUserSettingsIni);
+		GConfig->SetString(*ConfigSectionName, TEXT("PreviewAspectRatio"), *DefaultAspectRatio, GEditorPerProjectIni);
 		PreviewAspectRatio = DefaultAspectRatio;
 	}
 }
@@ -2061,9 +2061,9 @@ void SDesignerView::HandleOnCommonResolutionSelected(int32 Width, int32 Height,
 	PreviewHeight = Height;
 	PreviewAspectRatio = AspectRatio;
 
-	GConfig->SetInt(*ConfigSectionName, TEXT("PreviewWidth"), Width, GEditorUserSettingsIni);
-	GConfig->SetInt(*ConfigSectionName, TEXT("PreviewHeight"), Height, GEditorUserSettingsIni);
-	GConfig->SetString(*ConfigSectionName, TEXT("PreviewAspectRatio"), *AspectRatio, GEditorUserSettingsIni);
+	GConfig->SetInt(*ConfigSectionName, TEXT("PreviewWidth"), Width, GEditorPerProjectIni);
+	GConfig->SetInt(*ConfigSectionName, TEXT("PreviewHeight"), Height, GEditorPerProjectIni);
+	GConfig->SetString(*ConfigSectionName, TEXT("PreviewAspectRatio"), *AspectRatio, GEditorPerProjectIni);
 
 	if ( UUserWidget* DefaultWidget = GetDefaultWidget() )
 	{
diff --git a/Engine/Source/Editor/UMGEditor/Private/Palette/SPaletteView.cpp b/Engine/Source/Editor/UMGEditor/Private/Palette/SPaletteView.cpp
--- a/Engine/Source/Editor/UMGEditor/Private/Palette/SPaletteView.cpp
+++ b/Engine/Source/Editor/UMGEditor/Private/Palette/SPaletteView.cpp
@@ -271,7 +271,7 @@ void SPaletteView::LoadItemExpansion()
 	for ( TSharedPtr<FWidgetViewModel>& ViewModel : WidgetViewModels )
 	{
 		bool IsExpanded;
-		if ( GConfig->GetBool(TEXT("WidgetTemplatesExpanded"), *ViewModel->GetName().ToString(), IsExpanded, GEditorUserSettingsIni) && IsExpanded )
+		if ( GConfig->GetBool(TEXT("WidgetTemplatesExpanded"), *ViewModel->GetName().ToString(), IsExpanded, GEditorPerProjectIni) && IsExpanded )
 		{
 			WidgetTemplatesView->SetItemExpansion(ViewModel, true);
 		}
@@ -284,7 +284,7 @@ void SPaletteView::SaveItemExpansion()
 	for ( TSharedPtr<FWidgetViewModel>& ViewModel : WidgetViewModels )
 	{
 		const bool IsExpanded = WidgetTemplatesView->IsItemExpanded(ViewModel);
-		GConfig->SetBool(TEXT("WidgetTemplatesExpanded"), *ViewModel->GetName().ToString(), IsExpanded, GEditorUserSettingsIni);
+		GConfig->SetBool(TEXT("WidgetTemplatesExpanded"), *ViewModel->GetName().ToString(), IsExpanded, GEditorPerProjectIni);
 	}
 }
 
diff --git a/Engine/Source/Editor/UnrealEd/Classes/Editor/EditorEngine.h b/Engine/Source/Editor/UnrealEd/Classes/Editor/EditorEngine.h
--- a/Engine/Source/Editor/UnrealEd/Classes/Editor/EditorEngine.h
+++ b/Engine/Source/Editor/UnrealEd/Classes/Editor/EditorEngine.h
@@ -4,12 +4,13 @@
 
 #include "Engine/Brush.h"
 #include "Engine/Engine.h"
-#include "EditorUserSettings.h"
+#include "EditorPerProjectUserSettings.h"
 #include "Transactor.h"
 #include "../Settings/LevelEditorPlaySettings.h"
 #include "../Settings/LevelEditorViewportSettings.h"
 #include "EditorEngine.generated.h"
 
+class UEditorSettings;
 class APlayerStart;
 class FAssetData;
 class FPoly;
@@ -535,14 +536,6 @@ public:
 	UPROPERTY(transient)
 	class UTextureRenderTarget2D* ScratchRenderTarget256;
 
-	/** Global instance of the editor user settings */
-	UPROPERTY()
-	class UEditorUserSettings* EditorUserSettings;
-
-	/** Global instance of the game agnostic editor settings */
-	UPROPERTY()
-	class UEditorGameAgnosticSettings* GameAgnosticSettings;
-
 	/** A mesh component used to preview in editor without spawning a static mesh actor. */
 	UPROPERTY(transient)
 	class UStaticMeshComponent* PreviewMeshComp;
@@ -767,20 +760,6 @@ public:
 	 */
 	void InitBuilderBrush( UWorld* InWorld );
 
-	/** Returns the global instance of the editor user settings class. */
-	const UEditorUserSettings& GetEditorUserSettings() const;
-	UEditorUserSettings& AccessEditorUserSettings();
-
-	/** Saves the user settings to disk. */
-	void SaveEditorUserSettings();
-
-	/** Returns the global instance of the game agnostic editor settings class. */
-	const UEditorGameAgnosticSettings& GetGameAgnosticSettings() const;
-	UEditorGameAgnosticSettings& AccessGameAgnosticSettings();
-
-	/** Saves the game agnostic settings to disk. */
-	void SaveGameAgnosticSettings();
-
 	/** Access user setting for audio mute. */
 	bool IsRealTimeAudioMuted() const;
 
diff --git a/Engine/Source/Editor/UnrealEd/Classes/Factories/FbxAssetImportData.h b/Engine/Source/Editor/UnrealEd/Classes/Factories/FbxAssetImportData.h
--- a/Engine/Source/Editor/UnrealEd/Classes/Factories/FbxAssetImportData.h
+++ b/Engine/Source/Editor/UnrealEd/Classes/Factories/FbxAssetImportData.h
@@ -8,7 +8,7 @@
 /**
  * Base class for import data and options used when importing any asset from FBX
  */
-UCLASS(config=EditorUserSettings, HideCategories=Object, abstract)
+UCLASS(config=EditorPerProjectUserSettings, HideCategories=Object, abstract)
 class UNREALED_API UFbxAssetImportData : public UAssetImportData
 {
 	GENERATED_UCLASS_BODY()
diff --git a/Engine/Source/Editor/UnrealEd/Classes/Factories/FbxImportUI.h b/Engine/Source/Editor/UnrealEd/Classes/Factories/FbxImportUI.h
--- a/Engine/Source/Editor/UnrealEd/Classes/Factories/FbxImportUI.h
+++ b/Engine/Source/Editor/UnrealEd/Classes/Factories/FbxImportUI.h
@@ -22,7 +22,7 @@ enum EFBXImportType
 	FBXIT_MAX,
 };
 
-UCLASS(config=EditorUserSettings, AutoExpandCategories=(FTransform), HideCategories=Object, MinimalAPI)
+UCLASS(config=EditorPerProjectUserSettings, AutoExpandCategories=(FTransform), HideCategories=Object, MinimalAPI)
 class UFbxImportUI : public UObject
 {
 	GENERATED_UCLASS_BODY()
diff --git a/Engine/Source/Editor/UnrealEd/Classes/Factories/FbxMeshImportData.h b/Engine/Source/Editor/UnrealEd/Classes/Factories/FbxMeshImportData.h
--- a/Engine/Source/Editor/UnrealEd/Classes/Factories/FbxMeshImportData.h
+++ b/Engine/Source/Editor/UnrealEd/Classes/Factories/FbxMeshImportData.h
@@ -15,7 +15,7 @@ enum EFBXNormalImportMethod
 /**
  * Import data and options used when importing any mesh from FBX
  */
-UCLASS(config=EditorUserSettings, configdonotcheckdefaults, abstract)
+UCLASS(config=EditorPerProjectUserSettings, configdonotcheckdefaults, abstract)
 class UFbxMeshImportData : public UFbxAssetImportData
 {
 	GENERATED_UCLASS_BODY()
diff --git a/Engine/Source/Editor/UnrealEd/Classes/Factories/FbxStaticMeshImportData.h b/Engine/Source/Editor/UnrealEd/Classes/Factories/FbxStaticMeshImportData.h
--- a/Engine/Source/Editor/UnrealEd/Classes/Factories/FbxStaticMeshImportData.h
+++ b/Engine/Source/Editor/UnrealEd/Classes/Factories/FbxStaticMeshImportData.h
@@ -21,7 +21,7 @@ namespace EVertexColorImportOption
 	};
 }
 
-UCLASS(config=EditorUserSettings, AutoExpandCategories=(Options), MinimalAPI)
+UCLASS(config=EditorPerProjectUserSettings, AutoExpandCategories=(Options), MinimalAPI)
 class UFbxStaticMeshImportData : public UFbxMeshImportData
 {
 	GENERATED_UCLASS_BODY()
diff --git a/Engine/Source/Editor/UnrealEd/Classes/Preferences/CascadeOptions.h b/Engine/Source/Editor/UnrealEd/Classes/Preferences/CascadeOptions.h
--- a/Engine/Source/Editor/UnrealEd/Classes/Preferences/CascadeOptions.h
+++ b/Engine/Source/Editor/UnrealEd/Classes/Preferences/CascadeOptions.h
@@ -10,7 +10,7 @@
 #pragma once
 #include "CascadeOptions.generated.h"
 
-UCLASS(hidecategories=Object, config=EditorUserSettings)
+UCLASS(hidecategories=Object, config=EditorPerProjectUserSettings)
 class UNREALED_API UCascadeOptions : public UObject
 {
 	GENERATED_UCLASS_BODY()
diff --git a/Engine/Source/Editor/UnrealEd/Classes/Preferences/CurveEdOptions.h b/Engine/Source/Editor/UnrealEd/Classes/Preferences/CurveEdOptions.h
--- a/Engine/Source/Editor/UnrealEd/Classes/Preferences/CurveEdOptions.h
+++ b/Engine/Source/Editor/UnrealEd/Classes/Preferences/CurveEdOptions.h
@@ -10,7 +10,7 @@
 #pragma once
 #include "CurveEdOptions.generated.h"
 
-UCLASS(hidecategories=Object, config=EditorUserSettings)
+UCLASS(hidecategories=Object, config=EditorPerProjectUserSettings)
 class UNREALED_API UCurveEdOptions : public UObject
 {
 	GENERATED_UCLASS_BODY()
diff --git a/Engine/Source/Editor/UnrealEd/Classes/Settings/EditorLoadingSavingSettings.h b/Engine/Source/Editor/UnrealEd/Classes/Settings/EditorLoadingSavingSettings.h
--- a/Engine/Source/Editor/UnrealEd/Classes/Settings/EditorLoadingSavingSettings.h
+++ b/Engine/Source/Editor/UnrealEd/Classes/Settings/EditorLoadingSavingSettings.h
@@ -64,7 +64,7 @@ struct FAutoReimportDirectoryConfig
 /**
  * Implements the Level Editor's loading and saving settings.
  */
-UCLASS(config=EditorUserSettings, autoexpandcategories=(AutoSave, AutoReimport, Blueprints))
+UCLASS(config=EditorPerProjectUserSettings, autoexpandcategories=(AutoSave, AutoReimport, Blueprints))
 class UNREALED_API UEditorLoadingSavingSettings
 	: public UObject
 {
diff --git a/Engine/Source/Editor/UnrealEd/Private/Editor.cpp b/Engine/Source/Editor/UnrealEd/Private/Editor.cpp
--- a/Engine/Source/Editor/UnrealEd/Private/Editor.cpp
+++ b/Engine/Source/Editor/UnrealEd/Private/Editor.cpp
@@ -110,6 +110,8 @@
 #include "EngineStats.h"
 #include "Engine/SimpleConstructionScript.h"
 
+#include "Settings/EditorSettings.h"
+
 #define LOCTEXT_NAMESPACE "UnrealEd.Editor"
 
 TArray<FTickableEditorObject*> FTickableEditorObject::TickableObjects;
diff --git a/Engine/Source/Editor/UnrealEd/Private/EditorBuildUtils.cpp b/Engine/Source/Editor/UnrealEd/Private/EditorBuildUtils.cpp
--- a/Engine/Source/Editor/UnrealEd/Private/EditorBuildUtils.cpp
+++ b/Engine/Source/Editor/UnrealEd/Private/EditorBuildUtils.cpp
@@ -184,14 +184,14 @@ bool FEditorBuildUtils::EditorBuild( UWorld* InWorld, EBuildOptions::Type Id, co
 	if ( Id == EBuildOptions::BuildLighting )
 	{
 		// Retrieve settings from ini.
-		GConfig->GetBool( TEXT("LightingBuildOptions"), TEXT("OnlyBuildSelected"),		LightingBuildOptions.bOnlyBuildSelected,			GEditorUserSettingsIni );
-		GConfig->GetBool( TEXT("LightingBuildOptions"), TEXT("OnlyBuildCurrentLevel"),	LightingBuildOptions.bOnlyBuildCurrentLevel,		GEditorUserSettingsIni );
-		GConfig->GetBool( TEXT("LightingBuildOptions"), TEXT("OnlyBuildSelectedLevels"),LightingBuildOptions.bOnlyBuildSelectedLevels,	GEditorUserSettingsIni );
-		GConfig->GetBool( TEXT("LightingBuildOptions"), TEXT("OnlyBuildVisibility"),	LightingBuildOptions.bOnlyBuildVisibility,		GEditorUserSettingsIni );
-		GConfig->GetBool( TEXT("LightingBuildOptions"), TEXT("UseErrorColoring"),		LightingBuildOptions.bUseErrorColoring,			GEditorUserSettingsIni );
-		GConfig->GetBool( TEXT("LightingBuildOptions"), TEXT("ShowLightingBuildInfo"),	LightingBuildOptions.bShowLightingBuildInfo,		GEditorUserSettingsIni );
+		GConfig->GetBool( TEXT("LightingBuildOptions"), TEXT("OnlyBuildSelected"),		LightingBuildOptions.bOnlyBuildSelected,			GEditorPerProjectIni );
+		GConfig->GetBool( TEXT("LightingBuildOptions"), TEXT("OnlyBuildCurrentLevel"),	LightingBuildOptions.bOnlyBuildCurrentLevel,		GEditorPerProjectIni );
+		GConfig->GetBool( TEXT("LightingBuildOptions"), TEXT("OnlyBuildSelectedLevels"),LightingBuildOptions.bOnlyBuildSelectedLevels,	GEditorPerProjectIni );
+		GConfig->GetBool( TEXT("LightingBuildOptions"), TEXT("OnlyBuildVisibility"),	LightingBuildOptions.bOnlyBuildVisibility,		GEditorPerProjectIni );
+		GConfig->GetBool( TEXT("LightingBuildOptions"), TEXT("UseErrorColoring"),		LightingBuildOptions.bUseErrorColoring,			GEditorPerProjectIni );
+		GConfig->GetBool( TEXT("LightingBuildOptions"), TEXT("ShowLightingBuildInfo"),	LightingBuildOptions.bShowLightingBuildInfo,		GEditorPerProjectIni );
 		int32 QualityLevel;
-		GConfig->GetInt(  TEXT("LightingBuildOptions"), TEXT("QualityLevel"),			QualityLevel,						GEditorUserSettingsIni );
+		GConfig->GetInt(  TEXT("LightingBuildOptions"), TEXT("QualityLevel"),			QualityLevel,						GEditorPerProjectIni );
 		QualityLevel = FMath::Clamp<int32>(QualityLevel, Quality_Preview, Quality_Production);
 		LightingBuildOptions.QualityLevel = (ELightingBuildQuality)QualityLevel;
 	}
@@ -361,7 +361,7 @@ bool FEditorBuildUtils::EditorBuild( UWorld* InWorld, EBuildOptions::Type Id, co
 					}
 					else
 					{
-						GConfig->GetInt( TEXT("LightingBuildOptions"), TEXT("QualityLevel"), QualityLevel, GEditorUserSettingsIni);
+						GConfig->GetInt( TEXT("LightingBuildOptions"), TEXT("QualityLevel"), QualityLevel, GEditorPerProjectIni);
 						QualityLevel = FMath::Clamp<int32>(QualityLevel, Quality_Preview, Quality_Production);
 					}
 					LightingOptions.QualityLevel = (ELightingBuildQuality)QualityLevel;
diff --git a/Engine/Source/Editor/UnrealEd/Private/EditorDirectories.cpp b/Engine/Source/Editor/UnrealEd/Private/EditorDirectories.cpp
--- a/Engine/Source/Editor/UnrealEd/Private/EditorDirectories.cpp
+++ b/Engine/Source/Editor/UnrealEd/Private/EditorDirectories.cpp
@@ -29,18 +29,18 @@ void FEditorDirectories::LoadLastDirectories()
 	LastDir[ELastDirectory::PROJECT] = FPaths::RootDir();
 
 	// NOTE: We append a "2" to the section name to enforce backwards compatibility.  "Directories" is deprecated.
-	GConfig->GetString( TEXT("Directories2"), TEXT("UNR"),				LastDir[ELastDirectory::UNR],					GEditorUserSettingsIni );
-	GConfig->GetString( TEXT("Directories2"), TEXT("BRUSH"),			LastDir[ELastDirectory::BRUSH],				GEditorUserSettingsIni );
-	GConfig->GetString( TEXT("Directories2"), TEXT("FBX"),				LastDir[ELastDirectory::FBX],					GEditorUserSettingsIni );
-	GConfig->GetString( TEXT("Directories2"), TEXT("FBXAnim"),			LastDir[ELastDirectory::FBX_ANIM],			GEditorUserSettingsIni );
-	GConfig->GetString( TEXT("Directories2"), TEXT("GenericImport"),	LastDir[ELastDirectory::GENERIC_IMPORT],		GEditorUserSettingsIni );
-	GConfig->GetString( TEXT("Directories2"), TEXT("GenericExport"),	LastDir[ELastDirectory::GENERIC_EXPORT],		GEditorUserSettingsIni );
-	GConfig->GetString( TEXT("Directories2"), TEXT("GenericOpen"),		LastDir[ELastDirectory::GENERIC_OPEN],		GEditorUserSettingsIni );
-	GConfig->GetString( TEXT("Directories2"), TEXT("GenericSave"),		LastDir[ELastDirectory::GENERIC_SAVE],		GEditorUserSettingsIni );
-	GConfig->GetString( TEXT("Directories2"), TEXT("MeshImportExport"),	LastDir[ELastDirectory::MESH_IMPORT_EXPORT],	GEditorUserSettingsIni );
-	GConfig->GetString( TEXT("Directories2"), TEXT("WorldRoot"),		LastDir[ELastDirectory::WORLD_ROOT],			GEditorUserSettingsIni );
-	GConfig->GetString( TEXT("Directories2"), TEXT("Level"),			LastDir[ELastDirectory::LEVEL],					GEditorUserSettingsIni );
-	GConfig->GetString( TEXT("Directories2"), TEXT("Project"),			LastDir[ELastDirectory::PROJECT],				GEditorUserSettingsIni );
+	GConfig->GetString( TEXT("Directories2"), TEXT("UNR"),				LastDir[ELastDirectory::UNR],					GEditorPerProjectIni );
+	GConfig->GetString( TEXT("Directories2"), TEXT("BRUSH"),			LastDir[ELastDirectory::BRUSH],				GEditorPerProjectIni );
+	GConfig->GetString( TEXT("Directories2"), TEXT("FBX"),				LastDir[ELastDirectory::FBX],					GEditorPerProjectIni );
+	GConfig->GetString( TEXT("Directories2"), TEXT("FBXAnim"),			LastDir[ELastDirectory::FBX_ANIM],			GEditorPerProjectIni );
+	GConfig->GetString( TEXT("Directories2"), TEXT("GenericImport"),	LastDir[ELastDirectory::GENERIC_IMPORT],		GEditorPerProjectIni );
+	GConfig->GetString( TEXT("Directories2"), TEXT("GenericExport"),	LastDir[ELastDirectory::GENERIC_EXPORT],		GEditorPerProjectIni );
+	GConfig->GetString( TEXT("Directories2"), TEXT("GenericOpen"),		LastDir[ELastDirectory::GENERIC_OPEN],		GEditorPerProjectIni );
+	GConfig->GetString( TEXT("Directories2"), TEXT("GenericSave"),		LastDir[ELastDirectory::GENERIC_SAVE],		GEditorPerProjectIni );
+	GConfig->GetString( TEXT("Directories2"), TEXT("MeshImportExport"),	LastDir[ELastDirectory::MESH_IMPORT_EXPORT],	GEditorPerProjectIni );
+	GConfig->GetString( TEXT("Directories2"), TEXT("WorldRoot"),		LastDir[ELastDirectory::WORLD_ROOT],			GEditorPerProjectIni );
+	GConfig->GetString( TEXT("Directories2"), TEXT("Level"),			LastDir[ELastDirectory::LEVEL],					GEditorPerProjectIni );
+	GConfig->GetString( TEXT("Directories2"), TEXT("Project"),			LastDir[ELastDirectory::PROJECT],				GEditorPerProjectIni );
 
 }
 
@@ -48,18 +48,18 @@ void FEditorDirectories::LoadLastDirectories()
 void FEditorDirectories::SaveLastDirectories()
 {
 	// Save out default file directories
-	GConfig->SetString( TEXT("Directories2"), TEXT("UNR"),				*LastDir[ELastDirectory::UNR],				GEditorUserSettingsIni );
-	GConfig->SetString( TEXT("Directories2"), TEXT("BRUSH"),			*LastDir[ELastDirectory::BRUSH],				GEditorUserSettingsIni );
-	GConfig->SetString( TEXT("Directories2"), TEXT("FBX"),				*LastDir[ELastDirectory::FBX],				GEditorUserSettingsIni );
-	GConfig->SetString( TEXT("Directories2"), TEXT("FBXAnim"),			*LastDir[ELastDirectory::FBX_ANIM],			GEditorUserSettingsIni );
-	GConfig->SetString( TEXT("Directories2"), TEXT("GenericImport"),	*LastDir[ELastDirectory::GENERIC_IMPORT],		GEditorUserSettingsIni );
-	GConfig->SetString( TEXT("Directories2"), TEXT("GenericExport"),	*LastDir[ELastDirectory::GENERIC_EXPORT],		GEditorUserSettingsIni );
-	GConfig->SetString( TEXT("Directories2"), TEXT("GenericOpen"),		*LastDir[ELastDirectory::GENERIC_OPEN],		GEditorUserSettingsIni );
-	GConfig->SetString( TEXT("Directories2"), TEXT("GenericSave"),		*LastDir[ELastDirectory::GENERIC_SAVE],		GEditorUserSettingsIni );
-	GConfig->SetString( TEXT("Directories2"), TEXT("MeshImportExport"),	*LastDir[ELastDirectory::MESH_IMPORT_EXPORT],	GEditorUserSettingsIni );
-	GConfig->SetString( TEXT("Directories2"), TEXT("WorldRoot"),		*LastDir[ELastDirectory::WORLD_ROOT],			GEditorUserSettingsIni );
-	GConfig->SetString( TEXT("Directories2"), TEXT("Level"),			*LastDir[ELastDirectory::LEVEL],				GEditorUserSettingsIni );
-	GConfig->SetString( TEXT("Directories2"), TEXT("Project"),			*LastDir[ELastDirectory::PROJECT],				GEditorUserSettingsIni );
+	GConfig->SetString( TEXT("Directories2"), TEXT("UNR"),				*LastDir[ELastDirectory::UNR],				GEditorPerProjectIni );
+	GConfig->SetString( TEXT("Directories2"), TEXT("BRUSH"),			*LastDir[ELastDirectory::BRUSH],				GEditorPerProjectIni );
+	GConfig->SetString( TEXT("Directories2"), TEXT("FBX"),				*LastDir[ELastDirectory::FBX],				GEditorPerProjectIni );
+	GConfig->SetString( TEXT("Directories2"), TEXT("FBXAnim"),			*LastDir[ELastDirectory::FBX_ANIM],			GEditorPerProjectIni );
+	GConfig->SetString( TEXT("Directories2"), TEXT("GenericImport"),	*LastDir[ELastDirectory::GENERIC_IMPORT],		GEditorPerProjectIni );
+	GConfig->SetString( TEXT("Directories2"), TEXT("GenericExport"),	*LastDir[ELastDirectory::GENERIC_EXPORT],		GEditorPerProjectIni );
+	GConfig->SetString( TEXT("Directories2"), TEXT("GenericOpen"),		*LastDir[ELastDirectory::GENERIC_OPEN],		GEditorPerProjectIni );
+	GConfig->SetString( TEXT("Directories2"), TEXT("GenericSave"),		*LastDir[ELastDirectory::GENERIC_SAVE],		GEditorPerProjectIni );
+	GConfig->SetString( TEXT("Directories2"), TEXT("MeshImportExport"),	*LastDir[ELastDirectory::MESH_IMPORT_EXPORT],	GEditorPerProjectIni );
+	GConfig->SetString( TEXT("Directories2"), TEXT("WorldRoot"),		*LastDir[ELastDirectory::WORLD_ROOT],			GEditorPerProjectIni );
+	GConfig->SetString( TEXT("Directories2"), TEXT("Level"),			*LastDir[ELastDirectory::LEVEL],				GEditorPerProjectIni );
+	GConfig->SetString( TEXT("Directories2"), TEXT("Project"),			*LastDir[ELastDirectory::PROJECT],				GEditorPerProjectIni );
 }
 
 FString FEditorDirectories::GetLastDirectory( const ELastDirectory::Type InLastDir ) const
diff --git a/Engine/Source/Editor/UnrealEd/Private/EditorEngine.cpp b/Engine/Source/Editor/UnrealEd/Private/EditorEngine.cpp
--- a/Engine/Source/Editor/UnrealEd/Private/EditorEngine.cpp
+++ b/Engine/Source/Editor/UnrealEd/Private/EditorEngine.cpp
@@ -41,6 +41,8 @@
 #include "ContentBrowserModule.h"
 #include "ISourceCodeAccessModule.h"
 
+#include "Settings/EditorSettings.h"
+
 #include "Editor/MainFrame/Public/MainFrame.h"
 #include "AnimationUtils.h"
 #include "AudioDecompress.h"
@@ -411,7 +413,7 @@ static bool GetSmallToolBarIcons()
 
 static bool GetDisplayMultiboxHooks()
 {
-	return GEditor->AccessEditorUserSettings().bDisplayUIExtensionPoints;
+	return GetDefault<UEditorPerProjectUserSettings>()->bDisplayUIExtensionPoints;
 }
 
 void UEditorEngine::InitEditor(IEngineLoop* InEngineLoop)
@@ -1599,79 +1601,8 @@ float UEditorEngine::GetMaxTickRate( float DeltaTime, bool bAllowFrameRateSmooth
 	return MaxTickRate;
 }
 
-const UEditorUserSettings& UEditorEngine::GetEditorUserSettings() const
-{
-	if (EditorUserSettings == NULL)
-	{
-		auto ConstThis = const_cast< UEditorEngine* >( this );	// Hack because Header Generator doesn't yet support mutable keyword
-		ConstThis->EditorUserSettings = NewObject<UEditorUserSettings>();
-	}
-	return *EditorUserSettings;
-}
-
-UEditorUserSettings& UEditorEngine::AccessEditorUserSettings()
-{
-	if (EditorUserSettings == NULL)
-	{
-		EditorUserSettings = NewObject<UEditorUserSettings>();
-	}
-	return *EditorUserSettings;
-}
-
-void UEditorEngine::SaveEditorUserSettings()
-{
-	if (!FUnrealEdMisc::Get().IsDeletePreferences())
-	{
-		AccessEditorUserSettings().SaveConfig();
-	}
-}
-
-const UEditorGameAgnosticSettings& UEditorEngine::GetGameAgnosticSettings() const
-{
-	if (GameAgnosticSettings == NULL)
-	{
-		auto ConstThis = const_cast< UEditorEngine* >( this );	// Hack because Header Generator doesn't yet support mutable keyword
-		ConstThis->GameAgnosticSettings = NewObject<UEditorGameAgnosticSettings>();
-		
-		// Load config from file, but only if we are not the build machine since game agnostic settings may put the builder in an unclean state
-		if ( !GIsBuildMachine )
-		{
-			ConstThis->GameAgnosticSettings->LoadConfig(UEditorGameAgnosticSettings::StaticClass(), *GEditorGameAgnosticIni);
-		}
-	}
-	return *GameAgnosticSettings;
-}
-
-UEditorGameAgnosticSettings& UEditorEngine::AccessGameAgnosticSettings()
-{
-	if (GameAgnosticSettings == NULL)
-	{
-		GameAgnosticSettings = NewObject<UEditorGameAgnosticSettings>();
-		
-		// Load config from file, but only if we are not the build machine since game agnostic settings may put the builder in an unclean state
-		if ( !GIsBuildMachine )
-		{
-			GameAgnosticSettings->LoadConfig(UEditorGameAgnosticSettings::StaticClass(), *GEditorGameAgnosticIni);
-		}
-	}
-	return *GameAgnosticSettings;
-}
-
-void UEditorEngine::SaveGameAgnosticSettings()
-{
-	// Save config to file, but only if we are not the build machine since game agnostic settings may put the builder in an unclean state
-	if ( !GIsBuildMachine )
-	{
-		AccessGameAgnosticSettings().SaveConfig(CPF_Config, *GEditorGameAgnosticIni);
-	}
-}
-
 bool UEditorEngine::IsRealTimeAudioMuted() const
 {
-	if (EditorUserSettings == NULL)
-	{
-		return true;
-	}
 	return GetDefault<ULevelEditorMiscSettings>()->bEnableRealTimeAudio ? false : true;
 }
 
@@ -5378,7 +5309,7 @@ bool UEditorEngine::ShouldThrottleCPUUsage() const
 
 	if( !bIsForeground )
 	{
-		const UEditorUserSettings* Settings = GetDefault<UEditorUserSettings>();
+		const UEditorPerProjectUserSettings* Settings = GetDefault<UEditorPerProjectUserSettings>();
 		bShouldThrottle = Settings->bThrottleCPUWhenNotForeground;
 
 		// Check if we should throttle due to all windows being minimized
@@ -5864,7 +5795,7 @@ void UEditorEngine::UpdateRecentlyLoadedProjectFiles()
 	{
 		const FString AbsoluteProjectPath = IFileManager::Get().ConvertToAbsolutePathForExternalAppForRead(*FPaths::GetProjectFilePath());
 		// Update the recently loaded project files. Move this project file to the front of the list
-		TArray<FString>& RecentlyOpenedProjectFiles = AccessGameAgnosticSettings().RecentlyOpenedProjectFiles;
+		TArray<FString>& RecentlyOpenedProjectFiles = GetMutableDefault<UEditorSettings>()->RecentlyOpenedProjectFiles;
 		RecentlyOpenedProjectFiles.Remove( AbsoluteProjectPath );
 		RecentlyOpenedProjectFiles.Insert( AbsoluteProjectPath, 0 );
 
@@ -5885,7 +5816,7 @@ void UEditorEngine::UpdateRecentlyLoadedProjectFiles()
 			RecentlyOpenedProjectFiles.RemoveAt(MaxRecentProjectFiles, RecentlyOpenedProjectFiles.Num() - MaxRecentProjectFiles);
 		}
 
-		AccessGameAgnosticSettings().PostEditChange();
+		GetMutableDefault<UEditorSettings>()->PostEditChange();
 	}
 }
 
@@ -5914,7 +5845,7 @@ void UEditorEngine::UpdateAutoLoadProject()
 		bShouldLoadRecentProjects = false;
 	}
 
-	AccessGameAgnosticSettings().bLoadTheMostRecentlyLoadedProjectAtStartup = bShouldLoadRecentProjects;
+	GetMutableDefault<UEditorSettings>()->bLoadTheMostRecentlyLoadedProjectAtStartup = bShouldLoadRecentProjects;
 
 #if PLATFORM_MAC
 	if ( !GIsBuildMachine )
@@ -5935,7 +5866,7 @@ void UEditorEngine::UpdateAutoLoadProject()
 		{
 			if(FSlateApplication::IsInitialized())
 			{
-				FSuppressableWarningDialog::FSetupInfo Info( LOCTEXT("UpdateMacOSX_Body","Please update to the latest version of Mac OS X for best performance."), LOCTEXT("UpdateMacOSX_Title","Update Mac OS X"), TEXT("UpdateMacOSX"), GEditorGameAgnosticIni );
+				FSuppressableWarningDialog::FSetupInfo Info( LOCTEXT("UpdateMacOSX_Body","Please update to the latest version of Mac OS X for best performance."), LOCTEXT("UpdateMacOSX_Title","Update Mac OS X"), TEXT("UpdateMacOSX"), GEditorSettingsIni );
 				Info.ConfirmText = LOCTEXT( "OK", "OK");
 				Info.bDefaultToSupressInTheFuture = true;
 				FSuppressableWarningDialog OSUpdateWarning( Info );
@@ -5982,7 +5913,7 @@ void UEditorEngine::UpdateAutoLoadProject()
 		{
 			if(FSlateApplication::IsInitialized())
 			{
-				FSuppressableWarningDialog::FSetupInfo Info( LOCTEXT("UnsupportedGPUWarning_Body","The current graphics card does not meet the minimum specification, for best performance an NVIDIA GeForce 470 GTX or AMD Radeon 6870 HD series card or higher is recommended. Rendering performance and compatibility are not guaranteed with this graphics card."), LOCTEXT("UnsupportedGPUWarning_Title","Unsupported Graphics Card"), TEXT("UnsupportedGPUWarning"), GEditorGameAgnosticIni );
+				FSuppressableWarningDialog::FSetupInfo Info( LOCTEXT("UnsupportedGPUWarning_Body","The current graphics card does not meet the minimum specification, for best performance an NVIDIA GeForce 470 GTX or AMD Radeon 6870 HD series card or higher is recommended. Rendering performance and compatibility are not guaranteed with this graphics card."), LOCTEXT("UnsupportedGPUWarning_Title","Unsupported Graphics Card"), TEXT("UnsupportedGPUWarning"), GEditorSettingsIni );
 				Info.ConfirmText = LOCTEXT( "OK", "OK");
 				Info.bDefaultToSupressInTheFuture = true;
 				FSuppressableWarningDialog OSUpdateWarning( Info );
@@ -6029,7 +5960,7 @@ void UEditorEngine::UpdateAutoLoadProject()
 		{
 			if(FSlateApplication::IsInitialized())
 			{
-				FSuppressableWarningDialog::FSetupInfo Info( LOCTEXT("SlowGPUWarning_Body","The current graphics card is slower than the recommanded specification of an NVIDIA GeForce 470 GTX or AMD Radeon 6870 HD series card or higher, performance may be low."), LOCTEXT("SlowGPUWarning_Title","Slow Graphics Card"), TEXT("SlowGPUWarning"), GEditorGameAgnosticIni );
+				FSuppressableWarningDialog::FSetupInfo Info( LOCTEXT("SlowGPUWarning_Body","The current graphics card is slower than the recommanded specification of an NVIDIA GeForce 470 GTX or AMD Radeon 6870 HD series card or higher, performance may be low."), LOCTEXT("SlowGPUWarning_Title","Slow Graphics Card"), TEXT("SlowGPUWarning"), GEditorSettingsIni );
 				Info.ConfirmText = LOCTEXT( "OK", "OK");
 				Info.bDefaultToSupressInTheFuture = true;
 				FSuppressableWarningDialog OSUpdateWarning( Info );
@@ -6046,7 +5977,7 @@ void UEditorEngine::UpdateAutoLoadProject()
 		{
 			if(FSlateApplication::IsInitialized())
 			{
-				FSuppressableWarningDialog::FSetupInfo Info( LOCTEXT("LowRAMWarning_Body","For best performance install at least 8GB of RAM."), LOCTEXT("LowRAMWarning_Title","Low RAM"), TEXT("LowRAMWarning"), GEditorGameAgnosticIni );
+				FSuppressableWarningDialog::FSetupInfo Info( LOCTEXT("LowRAMWarning_Body","For best performance install at least 8GB of RAM."), LOCTEXT("LowRAMWarning_Title","Low RAM"), TEXT("LowRAMWarning"), GEditorSettingsIni );
 				Info.ConfirmText = LOCTEXT( "OK", "OK");
 				Info.bDefaultToSupressInTheFuture = true;
 				FSuppressableWarningDialog OSUpdateWarning( Info );
@@ -6063,7 +5994,7 @@ void UEditorEngine::UpdateAutoLoadProject()
 		{
 			if(FSlateApplication::IsInitialized())
 			{
-				FSuppressableWarningDialog::FSetupInfo Info( LOCTEXT("SlowCPUWarning_Body","For best performance a Quad-core Intel or AMD processor, 2.5 GHz or faster is recommended."), LOCTEXT("SlowCPUWarning_Title","CPU Performance Warning"), TEXT("SlowCPUWarning"), GEditorGameAgnosticIni );
+				FSuppressableWarningDialog::FSetupInfo Info( LOCTEXT("SlowCPUWarning_Body","For best performance a Quad-core Intel or AMD processor, 2.5 GHz or faster is recommended."), LOCTEXT("SlowCPUWarning_Title","CPU Performance Warning"), TEXT("SlowCPUWarning"), GEditorSettingsIni );
 				Info.ConfirmText = LOCTEXT( "OK", "OK");
 				Info.bDefaultToSupressInTheFuture = true;
 				FSuppressableWarningDialog OSUpdateWarning( Info );
diff --git a/Engine/Source/Editor/UnrealEd/Private/EditorModeManager.cpp b/Engine/Source/Editor/UnrealEd/Private/EditorModeManager.cpp
--- a/Engine/Source/Editor/UnrealEd/Private/EditorModeManager.cpp
+++ b/Engine/Source/Editor/UnrealEd/Private/EditorModeManager.cpp
@@ -78,12 +78,12 @@ FEditorModeTools::~FEditorModeTools()
 void FEditorModeTools::LoadConfig(void)
 {
 	GConfig->GetBool(TEXT("FEditorModeTools"),TEXT("ShowWidget"),bShowWidget,
-		GEditorUserSettingsIni);
+		GEditorPerProjectIni);
 
 	const bool bGetRawValue = true;
 	int32 Bogus = (int32)GetCoordSystem(bGetRawValue);
 	GConfig->GetInt(TEXT("FEditorModeTools"),TEXT("CoordSystem"),Bogus,
-		GEditorUserSettingsIni);
+		GEditorPerProjectIni);
 	SetCoordSystem((ECoordSystem)Bogus);
 
 
@@ -92,10 +92,10 @@ void FEditorModeTools::LoadConfig(void)
 
 void FEditorModeTools::SaveConfig(void)
 {
-	GConfig->SetBool(TEXT("FEditorModeTools"), TEXT("ShowWidget"), bShowWidget, GEditorUserSettingsIni);
+	GConfig->SetBool(TEXT("FEditorModeTools"), TEXT("ShowWidget"), bShowWidget, GEditorPerProjectIni);
 
 	const bool bGetRawValue = true;
-	GConfig->SetInt(TEXT("FEditorModeTools"), TEXT("CoordSystem"), (int32)GetCoordSystem(bGetRawValue), GEditorUserSettingsIni);
+	GConfig->SetInt(TEXT("FEditorModeTools"), TEXT("CoordSystem"), (int32)GetCoordSystem(bGetRawValue), GEditorPerProjectIni);
 
 	SaveWidgetSettings();
 }
@@ -820,7 +820,7 @@ void FEditorModeTools::CycleWidgetMode (void)
 /**Save Widget Settings to Ini file*/
 void FEditorModeTools::SaveWidgetSettings(void)
 {
-	GEditor->SaveEditorUserSettings();
+	GetMutableDefault<UEditorPerProjectUserSettings>()->SaveConfig();
 }
 
 /**Load Widget Settings from Ini file*/
diff --git a/Engine/Source/Editor/UnrealEd/Private/EditorServer.cpp b/Engine/Source/Editor/UnrealEd/Private/EditorServer.cpp
--- a/Engine/Source/Editor/UnrealEd/Private/EditorServer.cpp
+++ b/Engine/Source/Editor/UnrealEd/Private/EditorServer.cpp
@@ -69,6 +69,8 @@
 #include "UnrealEngine.h"
 #include "AI/Navigation/NavLinkRenderingComponent.h"
 
+#include "Settings/EditorSettings.h"
+
 DEFINE_LOG_CATEGORY_STATIC(LogEditorServer, Log, All);
 
 /** Used for the "tagsounds" and "checksounds" commands only			*/
@@ -1099,7 +1101,7 @@ void UEditorEngine::HandleTransactorUndo( FUndoSessionContext SessionContext, bo
 
 bool UEditorEngine::AreEditorAnalyticsEnabled() const 
 {
-	return GetGameAgnosticSettings().bEditorAnalyticsEnabled;
+	return GetDefault<UEditorSettings>()->bEditorAnalyticsEnabled;
 }
 
 void UEditorEngine::CreateStartupAnalyticsAttributes( TArray<FAnalyticsEventAttribute>& StartSessionAttributes ) const
@@ -1121,7 +1123,7 @@ UTransactor* UEditorEngine::CreateTrans()
 {
 	int32 UndoBufferSize;
 
-	if (!GConfig->GetInt(TEXT("Undo"), TEXT("UndoBufferSize"), UndoBufferSize, GEditorUserSettingsIni))
+	if (!GConfig->GetInt(TEXT("Undo"), TEXT("UndoBufferSize"), UndoBufferSize, GEditorPerProjectIni))
 	{
 		UndoBufferSize = 16;
 	}
diff --git a/Engine/Source/Editor/UnrealEd/Private/FileHelpers.cpp b/Engine/Source/Editor/UnrealEd/Private/FileHelpers.cpp
--- a/Engine/Source/Editor/UnrealEd/Private/FileHelpers.cpp
+++ b/Engine/Source/Editor/UnrealEd/Private/FileHelpers.cpp
@@ -1976,7 +1976,7 @@ void FEditorFileUtils::LoadMap(const FString& InFilename, bool LoadAsTemplate, b
 	}
 
 	// Save last opened level name.
-	GConfig->SetString(TEXT("EditorStartup"), TEXT("LastLevel"), *LongMapPackageName, GEditorUserSettingsIni);
+	GConfig->SetString(TEXT("EditorStartup"), TEXT("LastLevel"), *LongMapPackageName, GEditorPerProjectIni);
 
 	// Deactivate any editor modes when loading a new map
 	GLevelEditorModeTools().DeactivateAllModes();
@@ -3167,7 +3167,7 @@ void FEditorFileUtils::LoadDefaultMapAtStartup()
 	// Last opened map.
 	if (GetDefault<UEditorLoadingSavingSettings>()->LoadLevelAtStartup == ELoadLevelAtStartup::LastOpened)
 	{
-		GConfig->GetString(TEXT("EditorStartup"), TEXT("LastLevel"), EditorStartupMap, GEditorUserSettingsIni);
+		GConfig->GetString(TEXT("EditorStartup"), TEXT("LastLevel"), EditorStartupMap, GEditorPerProjectIni);
 	}
 	// Default project map.
 	if (EditorStartupMap.IsEmpty()) {
@@ -3262,7 +3262,7 @@ void FEditorFileUtils::FindAllSubmittableConfigFiles(TMap<FString, TSharedPtr<cl
 	for (const FString& ConfigFilename : ConfigFilenames)
 	{
 		// Only check files which are intended to be under source control
-		if (FPaths::GetCleanFilename(ConfigFilename) != TEXT("DefaultEditorUserSettings.ini"))
+		if (FPaths::GetCleanFilename(ConfigFilename) != TEXT("DefaultEditorPerProjectUserSettings.ini"))
 		{
 			FSourceControlStatePtr SourceControlState = SourceControlProvider.GetState(ConfigFilename, EStateCacheUsage::Use);
 
diff --git a/Engine/Source/Editor/UnrealEd/Private/Lightmass/Lightmass.cpp b/Engine/Source/Editor/UnrealEd/Private/Lightmass/Lightmass.cpp
--- a/Engine/Source/Editor/UnrealEd/Private/Lightmass/Lightmass.cpp
+++ b/Engine/Source/Editor/UnrealEd/Private/Lightmass/Lightmass.cpp
@@ -1781,7 +1781,7 @@ void FLightmassExporter::WriteSceneSettings( Lightmass::FSceneFileHeader& Scene
 		verify(GConfig->GetInt(TEXT("DevOptions.StaticLighting"), TEXT("MaxTriangleIrradiancePhotonCacheSamples"), Scene.GeneralSettings.MaxTriangleIrradiancePhotonCacheSamples, GLightmassIni));
 
 		int32 CheckQualityLevel;
-		GConfig->GetInt( TEXT("LightingBuildOptions"), TEXT("QualityLevel"), CheckQualityLevel, GEditorUserSettingsIni);
+		GConfig->GetInt( TEXT("LightingBuildOptions"), TEXT("QualityLevel"), CheckQualityLevel, GEditorPerProjectIni);
 		CheckQualityLevel = FMath::Clamp<int32>(CheckQualityLevel, Quality_Preview, Quality_Production);
 		UE_LOG(LogLightmassSolver, Log, TEXT("LIGHTMASS: Writing scene settings: Quality level %d (%d in INI)"), (int32)(QualityLevel), CheckQualityLevel);
 		if (CheckQualityLevel != QualityLevel)
@@ -1799,7 +1799,7 @@ void FLightmassExporter::WriteSceneSettings( Lightmass::FSceneFileHeader& Scene
 		default:
 			{
 				bool bUseErrorColoring = false;
-				GConfig->GetBool( TEXT("LightingBuildOptions"), TEXT("UseErrorColoring"),		bUseErrorColoring,					GEditorUserSettingsIni );
+				GConfig->GetBool( TEXT("LightingBuildOptions"), TEXT("UseErrorColoring"),		bUseErrorColoring,					GEditorPerProjectIni );
 				Scene.GeneralSettings.bUseErrorColoring = bUseErrorColoring;
 				if (bUseErrorColoring == false)
 				{
@@ -2598,7 +2598,7 @@ bool FLightmassProcessor::BeginRun()
 
 	// Check to see if swarm should be run minimized (it should by default)
 	bool bMinimizeSwarm = true;
-	GConfig->GetBool(TEXT("LightingBuildOptions"), TEXT("MinimizeSwarm"), bMinimizeSwarm, GEditorGameAgnosticIni);
+	GConfig->GetBool(TEXT("LightingBuildOptions"), TEXT("MinimizeSwarm"), bMinimizeSwarm, GEditorSettingsIni);
 	if ( bMinimizeSwarm )
 	{
 		UE_LOG(LogLightmassSolver, Log,  TEXT("Swarm will be run minimized") );
diff --git a/Engine/Source/Editor/UnrealEd/Private/MRUList.cpp b/Engine/Source/Editor/UnrealEd/Private/MRUList.cpp
--- a/Engine/Source/Editor/UnrealEd/Private/MRUList.cpp
+++ b/Engine/Source/Editor/UnrealEd/Private/MRUList.cpp
@@ -119,7 +119,7 @@ void FMRUList::InternalReadINI( TArray<FString>& OutItems, const FString& INISec
 	{
 		// Try to find data for a key formed as "INIKeyBaseItemIdx" for the provided INI section. If found, add the data to the output item array.
 		FString CurItem;
-		if ( GConfig->GetString( *INISection, *FString::Printf( TEXT("%s%d"), *INIKeyBase, ItemIdx ), CurItem, GEditorUserSettingsIni ) )
+		if ( GConfig->GetString( *INISection, *FString::Printf( TEXT("%s%d"), *INIKeyBase, ItemIdx ), CurItem, GEditorPerProjectIni ) )
 		{
 			OutItems.AddUnique( FPaths::ConvertRelativePathToFull(CurItem) );
 		}
@@ -129,14 +129,14 @@ void FMRUList::InternalReadINI( TArray<FString>& OutItems, const FString& INISec
 
 void FMRUList::InternalWriteINI( const TArray<FString>& InItems, const FString& INISection, const FString& INIKeyBase )
 {
-	GConfig->EmptySection( *INISection, GEditorUserSettingsIni );
+	GConfig->EmptySection( *INISection, GEditorPerProjectIni );
 
 	for ( int32 ItemIdx = 0; ItemIdx < InItems.Num(); ++ItemIdx )
 	{
-		GConfig->SetString( *INISection, *FString::Printf( TEXT("%s%d"), *INIKeyBase, ItemIdx ), *InItems[ ItemIdx ], GEditorUserSettingsIni );
+		GConfig->SetString( *INISection, *FString::Printf( TEXT("%s%d"), *INIKeyBase, ItemIdx ), *InItems[ ItemIdx ], GEditorPerProjectIni );
 	}
 
-	GConfig->Flush( false, GEditorUserSettingsIni );
+	GConfig->Flush( false, GEditorPerProjectIni );
 }
 
 
diff --git a/Engine/Source/Editor/UnrealEd/Private/PackageBackup.cpp b/Engine/Source/Editor/UnrealEd/Private/PackageBackup.cpp
--- a/Engine/Source/Editor/UnrealEd/Private/PackageBackup.cpp
+++ b/Engine/Source/Editor/UnrealEd/Private/PackageBackup.cpp
@@ -220,7 +220,7 @@ bool FAutoPackageBackup::ShouldBackupPackage( const UPackage& InPackage, FString
 bool FAutoPackageBackup::IsPackageBackupEnabled()
 {
 	bool bEnabled = false;
-	GConfig->GetBool( TEXT("FAutoPackageBackup"), TEXT("Enabled"), bEnabled, GEditorUserSettingsIni );
+	GConfig->GetBool( TEXT("FAutoPackageBackup"), TEXT("Enabled"), bEnabled, GEditorPerProjectIni );
 	return bEnabled;
 }
 
@@ -233,7 +233,7 @@ bool FAutoPackageBackup::IsPackageBackupEnabled()
 int32 FAutoPackageBackup::GetMaxAllowedBackupSpace()
 {
 	int32 MaxSpaceAllowed = 0;
-	if ( GConfig->GetInt( TEXT("FAutoPackageBackup"), TEXT("MaxAllowedSpaceInMB"), MaxSpaceAllowed, GEditorUserSettingsIni ) )
+	if ( GConfig->GetInt( TEXT("FAutoPackageBackup"), TEXT("MaxAllowedSpaceInMB"), MaxSpaceAllowed, GEditorPerProjectIni ) )
 	{
 		// Convert the user stored value from megabytes to bytes; <<= 20 is the same as *= 1024 * 1024
 		MaxSpaceAllowed <<= 20;
@@ -251,7 +251,7 @@ int32 FAutoPackageBackup::GetMaxAllowedBackupSpace()
 int32 FAutoPackageBackup::GetBackupInterval()
 {
 	int32 BackupInterval = 0;
-	if ( GConfig->GetInt( TEXT("FAutoPackageBackup"), TEXT("BackupIntervalInMinutes"), BackupInterval, GEditorUserSettingsIni ) )
+	if ( GConfig->GetInt( TEXT("FAutoPackageBackup"), TEXT("BackupIntervalInMinutes"), BackupInterval, GEditorPerProjectIni ) )
 	{
 		// Convert the user stored value from minutes to seconds
 		BackupInterval *= 60;
diff --git a/Engine/Source/Editor/UnrealEd/Private/PackageTools.cpp b/Engine/Source/Editor/UnrealEd/Private/PackageTools.cpp
--- a/Engine/Source/Editor/UnrealEd/Private/PackageTools.cpp
+++ b/Engine/Source/Editor/UnrealEd/Private/PackageTools.cpp
@@ -154,7 +154,7 @@ namespace PackageTools
 		bool bSuccessfullyCompleted = true;
 
 		// whether or not to suppress the ask to fully load message
-		bool bSuppress = GEditor->AccessEditorUserSettings().bSuppressFullyLoadPrompt;
+		bool bSuppress = GetDefault<UEditorPerProjectUserSettings>()->bSuppressFullyLoadPrompt;
 
 		// Make sure they are all fully loaded.
 		bool bNeedsUpdate = false;
diff --git a/Engine/Source/Editor/UnrealEd/Private/PerformanceMonitor.cpp b/Engine/Source/Editor/UnrealEd/Private/PerformanceMonitor.cpp
--- a/Engine/Source/Editor/UnrealEd/Private/PerformanceMonitor.cpp
+++ b/Engine/Source/Editor/UnrealEd/Private/PerformanceMonitor.cpp
@@ -8,6 +8,8 @@
 #include "SNotificationList.h"
 #include "NotificationManager.h"
 
+#include "Settings/EditorSettings.h"
+
 #define LOCTEXT_NAMESPACE "PerformanceMonitor"
 
 const double AutoApplyScalabilityTimeout = 10;
@@ -155,7 +157,7 @@ bool FPerformanceMonitor::WillAutoScalabilityHelp() const
 Scalability::FQualityLevels FPerformanceMonitor::GetAutoScalabilityQualityLevels() const
 {
 	const Scalability::FQualityLevels ExistingLevels = Scalability::GetQualityLevels();
-	Scalability::FQualityLevels NewLevels = GEditor->GetGameAgnosticSettings().EngineBenchmarkResult;
+	Scalability::FQualityLevels NewLevels = GetDefault<UEditorSettings>()->EngineBenchmarkResult;
 
 	// Make sure we don't turn settings up if the user has turned them down for any reason
 	NewLevels.ResolutionQuality		= FMath::Min(NewLevels.ResolutionQuality, ExistingLevels.ResolutionQuality);
@@ -171,12 +173,12 @@ Scalability::FQualityLevels FPerformanceMonitor::GetAutoScalabilityQualityLevels
 
 void FPerformanceMonitor::AutoApplyScalability()
 {
-	GEditor->AccessGameAgnosticSettings().AutoApplyScalabilityBenchmark();
+	GetMutableDefault<UEditorSettings>()->AutoApplyScalabilityBenchmark();
 
 	Scalability::FQualityLevels NewLevels = FPerformanceMonitor::GetAutoScalabilityQualityLevels();
 
 	Scalability::SetQualityLevels(NewLevels);
-	Scalability::SaveState(GEditorGameAgnosticIni);
+	Scalability::SaveState(GEditorSettingsIni);
 	GEditor->RedrawAllViewports();
 
 	const bool bAutoApplied = true;
@@ -216,10 +218,10 @@ void FPerformanceMonitor::ShowPerformanceWarning(FText MessageText)
 
 void FPerformanceMonitor::CancelPerformanceNotification()
 {
-	UEditorUserSettings& EditorUserSettings = GEditor->AccessEditorUserSettings();
-	EditorUserSettings.bMonitorEditorPerformance = false;
-	EditorUserSettings.PostEditChange();
-	EditorUserSettings.SaveConfig();
+	UEditorPerProjectUserSettings* EditorUserSettings = GetMutableDefault<UEditorPerProjectUserSettings>();
+	EditorUserSettings->bMonitorEditorPerformance = false;
+	EditorUserSettings->PostEditChange();
+	EditorUserSettings->SaveConfig();
 
 	Reset();
 }
@@ -249,7 +251,7 @@ void FPerformanceMonitor::Tick(float DeltaTime)
 	FineMovingAverage.Tick(FPlatformTime::Seconds(), GAverageFPS);
 	CoarseMovingAverage.Tick(FPlatformTime::Seconds(), GAverageFPS);
 
-	bool bMonitorEditorPerformance = GEditor->GetEditorUserSettings().bMonitorEditorPerformance;
+	bool bMonitorEditorPerformance = GetDefault<UEditorPerProjectUserSettings>()->bMonitorEditorPerformance;
 	if( !bMonitorEditorPerformance || !bIsNotificationAllowed )
 	{
 		return;
@@ -308,7 +310,7 @@ void FPerformanceMonitor::Tick(float DeltaTime)
 	}
 	else
 	{
-		if( GEditor->GetGameAgnosticSettings().IsScalabilityBenchmarkValid() )
+		if( GetDefault<UEditorSettings>()->IsScalabilityBenchmarkValid() )
 		{
 			return;
 		}
diff --git a/Engine/Source/Editor/UnrealEd/Private/PlayLevel.cpp b/Engine/Source/Editor/UnrealEd/Private/PlayLevel.cpp
--- a/Engine/Source/Editor/UnrealEd/Private/PlayLevel.cpp
+++ b/Engine/Source/Editor/UnrealEd/Private/PlayLevel.cpp
@@ -1772,7 +1772,7 @@ void UEditorEngine::PlayForMovieCapture()
 	{
 		bool bCloseEditor = false;
 
-		GConfig->GetBool(TEXT("MatineeCreateMovieOptions"), TEXT("CloseEditor"), bCloseEditor, GEditorUserSettingsIni);
+		GConfig->GetBool(TEXT("MatineeCreateMovieOptions"), TEXT("CloseEditor"), bCloseEditor, GEditorPerProjectIni);
 
 		if (bCloseEditor)
 		{
diff --git a/Engine/Source/Editor/UnrealEd/Private/SScalabilitySettings.cpp b/Engine/Source/Editor/UnrealEd/Private/SScalabilitySettings.cpp
--- a/Engine/Source/Editor/UnrealEd/Private/SScalabilitySettings.cpp
+++ b/Engine/Source/Editor/UnrealEd/Private/SScalabilitySettings.cpp
@@ -3,6 +3,8 @@
 #include "UnrealEd.h"
 #include "SScalabilitySettings.h"
 
+#include "Settings/EditorSettings.h"
+
 #define LOCTEXT_NAMESPACE "EngineScalabiltySettings"
 
 ECheckBoxState SScalabilitySettings::IsGroupQualityLevelSelected(const TCHAR* InGroupName, int32 InQualityLevel) const
@@ -31,7 +33,7 @@ void SScalabilitySettings::OnGroupQualityLevelChanged(ECheckBoxState NewState, c
 	else if (FCString::Strcmp(InGroupName, TEXT("EffectsQuality")) == 0) CachedQualityLevels.EffectsQuality = InQualityLevel;
 
 	Scalability::SetQualityLevels(CachedQualityLevels);
-	Scalability::SaveState(GEditorGameAgnosticIni);
+	Scalability::SaveState(GEditorSettingsIni);
 	GEditor->RedrawAllViewports();
 }
 
@@ -40,7 +42,7 @@ void SScalabilitySettings::OnResolutionScaleChanged(float InValue)
 	CachedQualityLevels.ResolutionQuality = (int32)(FMath::Lerp(Scalability::MinResolutionScale, Scalability::MaxResolutionScale, InValue));
 
 	Scalability::SetQualityLevels(CachedQualityLevels);
-	Scalability::SaveState(GEditorGameAgnosticIni);
+	Scalability::SaveState(GEditorSettingsIni);
 	GEditor->RedrawAllViewports();
 }
 
@@ -96,17 +98,18 @@ FReply SScalabilitySettings::OnHeaderClicked(int32 InQualityLevel)
 {
 	CachedQualityLevels.SetFromSingleQualityLevel(InQualityLevel);
 	Scalability::SetQualityLevels(CachedQualityLevels);
-	Scalability::SaveState(GEditorGameAgnosticIni);
+	Scalability::SaveState(GEditorSettingsIni);
 	GEditor->RedrawAllViewports();
 	return FReply::Handled();
 }
 
 FReply SScalabilitySettings::OnAutoClicked()
 {
-	GEditor->AccessGameAgnosticSettings().AutoApplyScalabilityBenchmark();
-	GEditor->AccessGameAgnosticSettings().LoadScalabilityBenchmark();
+	auto* Settings = GetMutableDefault<UEditorSettings>();
+	Settings->AutoApplyScalabilityBenchmark();
+	Settings->LoadScalabilityBenchmark();
 
-	CachedQualityLevels = GEditor->GetGameAgnosticSettings().EngineBenchmarkResult;
+	CachedQualityLevels = Settings->EngineBenchmarkResult;
 
 	GEditor->RedrawAllViewports();
 	return FReply::Handled();
@@ -124,7 +127,7 @@ SGridPanel::FSlot& SScalabilitySettings::MakeGridSlot(int32 InCol, int32 InRow,
 
 ECheckBoxState SScalabilitySettings::IsMonitoringPerformance() const
 {
-	const bool bMonitorEditorPerformance = GEditor->GetEditorUserSettings().bMonitorEditorPerformance;
+	const bool bMonitorEditorPerformance = GetDefault<UEditorPerProjectUserSettings>()->bMonitorEditorPerformance;
 	return bMonitorEditorPerformance ? ECheckBoxState::Checked : ECheckBoxState::Unchecked;
 }
 
@@ -132,10 +135,10 @@ void SScalabilitySettings::OnMonitorPerformanceChanged(ECheckBoxState NewState)
 {
 	const bool bNewEnabledState = ( NewState == ECheckBoxState::Checked );
 
-	auto& Settings = GEditor->AccessEditorUserSettings();
-	Settings.bMonitorEditorPerformance = bNewEnabledState;
-	Settings.PostEditChange();
-	Settings.SaveConfig();
+	auto* Settings = GetMutableDefault<UEditorPerProjectUserSettings>();
+	Settings->bMonitorEditorPerformance = bNewEnabledState;
+	Settings->PostEditChange();
+	Settings->SaveConfig();
 }
 
 void SScalabilitySettings::Construct( const FArguments& InArgs )
diff --git a/Engine/Source/Editor/UnrealEd/Private/Settings/EditorLoadingSavingSettingsCustomization.h b/Engine/Source/Editor/UnrealEd/Private/Settings/EditorLoadingSavingSettingsCustomization.h
--- a/Engine/Source/Editor/UnrealEd/Private/Settings/EditorLoadingSavingSettingsCustomization.h
+++ b/Engine/Source/Editor/UnrealEd/Private/Settings/EditorLoadingSavingSettingsCustomization.h
@@ -5,6 +5,7 @@
 #include "PropertyEditing.h"
 #include "PropertyCustomizationHelpers.h"
 
+#include "Settings/EditorSettings.h"
 
 #define LOCTEXT_NAMESPACE "FEditorLoadingSavingSettingsCustomization"
 
@@ -56,7 +57,7 @@ protected:
 		IDetailCategoryBuilder& StartupCategory = LayoutBuilder.EditCategory("Startup");
 		{
 			TArray<UObject*> ObjectList;
-			ObjectList.Add(&GEditor->AccessGameAgnosticSettings());
+			ObjectList.Add(GetMutableDefault<UEditorSettings>());
 			StartupCategory.AddExternalProperty(ObjectList, "bLoadTheMostRecentlyLoadedProjectAtStartup");
 
 			StartupCategory.AddExternalProperty(ObjectList, "bEditorAnalyticsEnabled", EPropertyLocation::Advanced);
diff --git a/Engine/Source/Editor/UnrealEd/Private/Settings/EditorPerProjectUserSettings.cpp b/Engine/Source/Editor/UnrealEd/Private/Settings/EditorPerProjectUserSettings.cpp
new file mode 100644
--- /dev/null
+++ b/Engine/Source/Editor/UnrealEd/Private/Settings/EditorPerProjectUserSettings.cpp
@@ -0,0 +1,49 @@
+// Copyright 1998-2015 Epic Games, Inc. All Rights Reserved.
+
+#include "UnrealEd.h"
+#include "MessageLog.h"
+#include "BlueprintPaletteFavorites.h"
+
+#define LOCTEXT_NAMESPACE "EditorPerProjectUserSettings"
+
+UEditorPerProjectUserSettings::UEditorPerProjectUserSettings(const FObjectInitializer& ObjectInitializer)
+	: Super(ObjectInitializer)
+{
+	//Default to high quality
+	MaterialQualityLevel = 1;
+	bMonitorEditorPerformance = true;
+	BlueprintFavorites = CreateDefaultSubobject<UBlueprintPaletteFavorites>(TEXT("BlueprintFavorites"));
+}
+
+void UEditorPerProjectUserSettings::PostInitProperties()
+{
+	Super::PostInitProperties();
+
+	//Ensure the material quality cvar is set to the settings loaded.
+	static IConsoleVariable* MaterialQualityLevelVar = IConsoleManager::Get().FindConsoleVariable(TEXT("r.MaterialQualityLevel"));
+	MaterialQualityLevelVar->Set(MaterialQualityLevel, ECVF_SetByScalability);
+}
+
+void UEditorPerProjectUserSettings::PostEditChangeProperty( FPropertyChangedEvent& PropertyChangedEvent )
+{
+	Super::PostEditChangeProperty(PropertyChangedEvent);
+
+	const FName Name = (PropertyChangedEvent.Property != nullptr) ? PropertyChangedEvent.Property->GetFName() : NAME_None;
+
+	if (Name == FName(TEXT("bUseCurvesForDistributions")))
+	{
+		extern CORE_API uint32 GDistributionType;
+		//GDistributionType == 0 for curves
+		GDistributionType = (bUseCurvesForDistributions) ? 0 : 1;
+	}
+
+	if (!FUnrealEdMisc::Get().IsDeletePreferences())
+	{
+		SaveConfig();
+	}
+
+	UserSettingChangedEvent.Broadcast(Name);
+}
+
+
+#undef LOCTEXT_NAMESPACE
