Fixed matinee crash when splitting translation and rotation tracks in matinee while a curve in the curve editor is selected

https://jira.ol.epicgames.net/browse/UE-19008

diff --git a/Engine/Source/Editor/Matinee/Private/MatineeMenus.cpp b/Engine/Source/Editor/Matinee/Private/MatineeMenus.cpp
--- a/Engine/Source/Editor/Matinee/Private/MatineeMenus.cpp
+++ b/Engine/Source/Editor/Matinee/Private/MatineeMenus.cpp
@@ -4273,6 +4273,7 @@ void FMatinee::OnSplitTranslationAndRotation()
 	check( HasATrackSelected() );
 
 	ClearKeySelection();
+
 	// Check to make sure there is a movement track in list
 	// before attempting to start the transaction system.
 	if( HasATrackSelected( UInterpTrackMove::StaticClass() ) )
@@ -4299,6 +4300,9 @@ void FMatinee::OnSplitTranslationAndRotation()
 
 		UpdateTrackWindowScrollBars();
 	}
+	
+	// Make sure the curve editor is in sync
+	CurveEd->CurveChanged();
 }
 
 /** Normalize Velocity Dialog */
