Set default slot name to be default.

diff --git a/Engine/Source/Editor/Matinee/Private/MatineeTrackHelpers.cpp b/Engine/Source/Editor/Matinee/Private/MatineeTrackHelpers.cpp
--- a/Engine/Source/Editor/Matinee/Private/MatineeTrackHelpers.cpp
+++ b/Engine/Source/Editor/Matinee/Private/MatineeTrackHelpers.cpp
@@ -84,7 +84,7 @@ UMatineeTrackAnimControlHelper::UMatineeTrackAnimControlHelper(const FObjectInit
 bool UMatineeTrackAnimControlHelper::PreCreateTrack( UInterpGroup* Group, const UInterpTrack *TrackDef, bool bDuplicatingTrack, bool bAllowPrompts ) const
 {
 	// For AnimControl tracks - pop up a dialog to choose slot name.
-	AnimSlotName = NAME_None;
+	AnimSlotName = FAnimSlotGroup::DefaultSlotName;
 
 	FEdModeInterpEdit* Mode = (FEdModeInterpEdit*)GLevelEditorModeTools().GetActiveMode( FBuiltinEditorModes::EM_InterpEdit );
 	check(Mode != NULL);
