Renames "Cinematic Preview" to "Allow Cinematic Control" to remove confusion with the difference between Cinematic Preview and Cinematic Viewport.

diff --git a/Engine/Plugins/MovieScene/LevelSequenceEditor/Source/LevelSequenceEditor/Private/CinematicViewport/SCinematicLevelViewport.cpp b/Engine/Plugins/MovieScene/LevelSequenceEditor/Source/LevelSequenceEditor/Private/CinematicViewport/SCinematicLevelViewport.cpp
--- a/Engine/Plugins/MovieScene/LevelSequenceEditor/Source/LevelSequenceEditor/Private/CinematicViewport/SCinematicLevelViewport.cpp
+++ b/Engine/Plugins/MovieScene/LevelSequenceEditor/Source/LevelSequenceEditor/Private/CinematicViewport/SCinematicLevelViewport.cpp
@@ -86,7 +86,7 @@ FCinematicViewportClient::FCinematicViewportClient()
 	bDrawAxes = false;
 	bIsRealtime = true;
 	SetGameView(true);
-	SetAllowCinematicPreview(true);
+	SetAllowCinematicControl(true);
 	bDisableInput = false;
 }
 
diff --git a/Engine/Source/Editor/LevelEditor/Private/LevelViewportActions.cpp b/Engine/Source/Editor/LevelEditor/Private/LevelViewportActions.cpp
--- a/Engine/Source/Editor/LevelEditor/Private/LevelViewportActions.cpp
+++ b/Engine/Source/Editor/LevelEditor/Private/LevelViewportActions.cpp
@@ -58,7 +58,7 @@ void FLevelViewportCommands::RegisterCommands()
 
 	UI_COMMAND( ApplyMaterialToActor, "Apply Material", "Attempts to apply a dropped material to this object", EUserInterfaceActionType::Button, FInputChord() );
 
-	UI_COMMAND( ToggleCinematicPreview, "Cinematic Preview", "If enabled, allows Matinee or Sequencer previews to play in this viewport", EUserInterfaceActionType::ToggleButton, FInputChord() );
+	UI_COMMAND( ToggleCinematicPreview, "Allow Cinematic Control", "If enabled, allows Matinee or Sequencer previews to play in this viewport", EUserInterfaceActionType::ToggleButton, FInputChord() );
 
 	UI_COMMAND( FindInLevelScriptBlueprint, "Find In Level Script", "Finds references of a selected actor in the level script blueprint", EUserInterfaceActionType::Button, FInputChord(EModifierKey::Control, EKeys::K) );
 	UI_COMMAND( AdvancedSettings, "Advanced Settings...", "Opens the advanced viewport settings", EUserInterfaceActionType::Button, FInputChord());
diff --git a/Engine/Source/Editor/LevelEditor/Private/SLevelViewport.cpp b/Engine/Source/Editor/LevelEditor/Private/SLevelViewport.cpp
--- a/Engine/Source/Editor/LevelEditor/Private/SLevelViewport.cpp
+++ b/Engine/Source/Editor/LevelEditor/Private/SLevelViewport.cpp
@@ -431,7 +431,7 @@ void SLevelViewport::ConstructLevelEditorViewportClient( const FArguments& InArg
 	{
 		LevelViewportClient->SetViewLocation( EditorViewportDefs::DefaultPerspectiveViewLocation );
 		LevelViewportClient->SetViewRotation( EditorViewportDefs::DefaultPerspectiveViewRotation );
-		LevelViewportClient->SetAllowCinematicPreview(true);
+		LevelViewportClient->SetAllowCinematicControl(true);
 	}
 	LevelViewportClient->SetRealtime(ViewportInstanceSettings.bIsRealtime);
 	LevelViewportClient->SetShowStats(ViewportInstanceSettings.bShowOnScreenStats);
@@ -2133,15 +2133,15 @@ void SLevelViewport::OnClearAllBookMarks()
 void SLevelViewport::OnToggleAllowCinematicPreview()
 {
 	// Reset the FOV of Viewport for cases where we have been previewing the matinee with a changing FOV
-	LevelViewportClient->ViewFOV = LevelViewportClient->AllowsCinematicPreview() ? LevelViewportClient->ViewFOV : LevelViewportClient->FOVAngle;
+	LevelViewportClient->ViewFOV = LevelViewportClient->AllowsCinematicControl() ? LevelViewportClient->ViewFOV : LevelViewportClient->FOVAngle;
 
-	LevelViewportClient->SetAllowCinematicPreview( !LevelViewportClient->AllowsCinematicPreview() );
+	LevelViewportClient->SetAllowCinematicControl( !LevelViewportClient->AllowsCinematicControl() );
 	LevelViewportClient->Invalidate( false );
 }
 
 bool SLevelViewport::AllowsCinematicPreview() const
 {
-	return LevelViewportClient->AllowsCinematicPreview();
+	return LevelViewportClient->AllowsCinematicControl();
 }
 
 void SLevelViewport::OnIncrementPositionGridSize()
@@ -3081,7 +3081,7 @@ void SLevelViewport::PreviewActors( const TArray< AActor* >& ActorsToPreview )
 				ActorPreviewLevelViewportClient->bDisableInput = true;
 
 				// Never allow Matinee to possess these views
-				ActorPreviewLevelViewportClient->SetAllowCinematicPreview( false );
+				ActorPreviewLevelViewportClient->SetAllowCinematicControl( false );
 
 				// Our preview viewport is always visible if our owning SLevelViewport is visible, so we hook up
 				// to the same IsVisible method
diff --git a/Engine/Source/Editor/Matinee/Private/Matinee.cpp b/Engine/Source/Editor/Matinee/Private/Matinee.cpp
--- a/Engine/Source/Editor/Matinee/Private/Matinee.cpp
+++ b/Engine/Source/Editor/Matinee/Private/Matinee.cpp
@@ -357,7 +357,7 @@ void FMatinee::SetAudioRealtimeOverride( bool bAudioIsRealtime ) const
 		FEditorViewportClient* const LevelVC = GEditor->LevelViewportClients[i];
 		if (LevelVC)
 		{
-			if(LevelVC->IsPerspective() && LevelVC->AllowsCinematicPreview() )
+			if(LevelVC->IsPerspective() && LevelVC->AllowsCinematicControl() )
 			{
 				LevelVC->SetForcedAudioRealtime(bAudioIsRealtime);
 			}
@@ -369,7 +369,7 @@ void FMatinee::OnToggleAspectRatioBars()
 {
 	if (GCurrentLevelEditingViewportClient)
 	{
-		if(GCurrentLevelEditingViewportClient->IsPerspective() && GCurrentLevelEditingViewportClient->AllowsCinematicPreview() )
+		if(GCurrentLevelEditingViewportClient->IsPerspective() && GCurrentLevelEditingViewportClient->AllowsCinematicControl() )
 		{
 			bool bEnabled = !AreAspectRatioBarsEnabled();
 			GCurrentLevelEditingViewportClient->SetShowAspectRatioBarDisplay(bEnabled);
@@ -383,7 +383,7 @@ void FMatinee::OnToggleSafeFrames()
 {
 	if (GCurrentLevelEditingViewportClient)
 	{
-		if(GCurrentLevelEditingViewportClient->IsPerspective() && GCurrentLevelEditingViewportClient->AllowsCinematicPreview() )
+		if(GCurrentLevelEditingViewportClient->IsPerspective() && GCurrentLevelEditingViewportClient->AllowsCinematicControl() )
 		{
 			bool bEnabled = !IsSafeFrameDisplayEnabled();
 			GCurrentLevelEditingViewportClient->SetShowSafeFrameBoxDisplay(bEnabled);
@@ -956,7 +956,7 @@ void FMatinee::InitMatinee(const EToolkitMode::Type Mode, const TSharedPtr< clas
 		if(LevelVC)
 		{
 			// If there is a director group, set the perspective viewports to realtime automatically.
-			if(LevelVC->IsPerspective() && LevelVC->AllowsCinematicPreview())
+			if(LevelVC->IsPerspective() && LevelVC->AllowsCinematicControl())
 			{				
 				//Ensure Realtime is turned on and store the original setting so we can restore it later.
 				LevelVC->SetRealtime(true, true);
@@ -2085,7 +2085,7 @@ void FMatinee::OnClose()
 		if(LevelVC)
 		{
 			// Turn off realtime when exiting.
-			if( LevelVC->IsPerspective() && LevelVC->AllowsCinematicPreview() )
+			if( LevelVC->IsPerspective() && LevelVC->AllowsCinematicControl() )
 			{				
 				//Specify true so RestoreRealtime will allow us to disable Realtime if it was original disabled
 				LevelVC->RestoreRealtime(true);
@@ -2198,7 +2198,7 @@ static void DrawTransparentLine( FCanvas* Canvas, const FVector2D& Start, const
 
 void FMatinee::DrawModeHUD(FEditorViewportClient* ViewportClient,FViewport* Viewport,const FSceneView* View,FCanvas* Canvas)
 {
-	if( ViewportClient->AllowsCinematicPreview() )
+	if( ViewportClient->AllowsCinematicControl() )
 	{
 		// Get the size of the viewport
 		const int32 SizeX = Viewport->GetSizeXY().X;
diff --git a/Engine/Source/Editor/Matinee/Private/MatineeTools.cpp b/Engine/Source/Editor/Matinee/Private/MatineeTools.cpp
--- a/Engine/Source/Editor/Matinee/Private/MatineeTools.cpp
+++ b/Engine/Source/Editor/Matinee/Private/MatineeTools.cpp
@@ -124,7 +124,7 @@ void FMatinee::UpdateViewportSettings()
 {
 	if ( GCurrentLevelEditingViewportClient )
 	{
-		if ( GCurrentLevelEditingViewportClient->IsPerspective() && GCurrentLevelEditingViewportClient->AllowsCinematicPreview() )
+		if ( GCurrentLevelEditingViewportClient->IsPerspective() && GCurrentLevelEditingViewportClient->AllowsCinematicControl() )
 		{
 			bool bSafeFrames = IsSafeFrameDisplayEnabled();
 			bool bAspectRatioBars = AreAspectRatioBarsEnabled();
@@ -3776,7 +3776,7 @@ void FMatinee::LockCamToGroup(class UInterpGroup* InGroup, const bool bResetView
 			for(int32 i=0; i<GEditor->LevelViewportClients.Num(); i++)
 			{
 				FLevelEditorViewportClient* LevelVC =GEditor->LevelViewportClients[i];
-				if(LevelVC && LevelVC->IsPerspective() && LevelVC->AllowsCinematicPreview() )
+				if(LevelVC && LevelVC->IsPerspective() && LevelVC->AllowsCinematicControl() )
 				{
 					LevelVC->RemoveCameraRoll();
 					LevelVC->ViewFOV = LevelVC->FOVAngle;
@@ -3861,7 +3861,7 @@ void FMatinee::UpdateCameraToGroup(const bool bInUpdateStandardViewports, bool b
 		for(int32 i=0; i<GEditor->LevelViewportClients.Num(); i++)
 		{
 			FLevelEditorViewportClient* LevelVC =GEditor->LevelViewportClients[i];
-			if(LevelVC && LevelVC->IsPerspective() && LevelVC->AllowsCinematicPreview() )
+			if(LevelVC && LevelVC->IsPerspective() && LevelVC->AllowsCinematicControl() )
 			{
 				UpdateLevelViewport(DefaultViewedActor, LevelVC, FadeAmount, ColorScale, bEnableColorScaling, bUpdateViewportTransform );
 			}
@@ -3967,7 +3967,7 @@ void FMatinee::SaveLevelViewports()
 	for( int32 ViewIndex = 0; ViewIndex < GEditor->LevelViewportClients.Num(); ++ViewIndex )
 	{
 		FLevelEditorViewportClient* LevelVC = GEditor->LevelViewportClients[ ViewIndex ];
-		if( LevelVC && LevelVC->IsPerspective() && LevelVC->AllowsCinematicPreview() )
+		if( LevelVC && LevelVC->IsPerspective() && LevelVC->AllowsCinematicControl() )
 		{
 			FMatineeViewSaveData SaveData;
 			SaveData.ViewIndex = ViewIndex;
@@ -3988,7 +3988,7 @@ void FMatinee::RestoreLevelViewports()
 		if( SavedData.ViewIndex < GEditor->LevelViewportClients.Num() )
 		{
 			FLevelEditorViewportClient* LevelVC = GEditor->LevelViewportClients[ SavedData.ViewIndex ];
-			if ( LevelVC && LevelVC->IsPerspective() && LevelVC->AllowsCinematicPreview() )
+			if ( LevelVC && LevelVC->IsPerspective() && LevelVC->AllowsCinematicControl() )
 			{
 				LevelVC->SetMatineeActorLock( nullptr );
 				LevelVC->SetViewRotation( SavedData.ViewRotation );
diff --git a/Engine/Source/Editor/Matinee/Private/SMatineeRecorder.cpp b/Engine/Source/Editor/Matinee/Private/SMatineeRecorder.cpp
--- a/Engine/Source/Editor/Matinee/Private/SMatineeRecorder.cpp
+++ b/Engine/Source/Editor/Matinee/Private/SMatineeRecorder.cpp
@@ -128,7 +128,7 @@ void SMatineeRecorder::Construct(const FArguments& InArgs)
 	LevelViewportClient->SetViewRotation( EditorViewportDefs::DefaultPerspectiveViewRotation );
 
 	LevelViewportClient->SetRealtime( true );
-	LevelViewportClient->SetAllowCinematicPreview( true );
+	LevelViewportClient->SetAllowCinematicControl( true );
 	LevelViewportClient->Viewport->SetUserFocus(true);
 	LevelViewportClient->SetMatineeRecordingWindow( ParentMatineeWindow.Pin().Get() );
 
diff --git a/Engine/Source/Editor/MovieSceneTools/Private/TrackEditorThumbnail/TrackEditorThumbnail.cpp b/Engine/Source/Editor/MovieSceneTools/Private/TrackEditorThumbnail/TrackEditorThumbnail.cpp
--- a/Engine/Source/Editor/MovieSceneTools/Private/TrackEditorThumbnail/TrackEditorThumbnail.cpp
+++ b/Engine/Source/Editor/MovieSceneTools/Private/TrackEditorThumbnail/TrackEditorThumbnail.cpp
@@ -663,7 +663,7 @@ void FTrackEditorThumbnailCache::Setup()
 			InternalViewportClient->ViewportType = LVT_Perspective;
 			InternalViewportClient->bDisableInput = true;
 			InternalViewportClient->bDrawAxes = false;
-			InternalViewportClient->SetAllowCinematicPreview(false);
+			InternalViewportClient->SetAllowCinematicControl(false);
 			InternalViewportClient->SetRealtime(false);
 
 			SetupViewportEngineFlags();
diff --git a/Engine/Source/Editor/MovieSceneTools/Private/TrackEditors/CameraCutTrackEditor.cpp b/Engine/Source/Editor/MovieSceneTools/Private/TrackEditors/CameraCutTrackEditor.cpp
--- a/Engine/Source/Editor/MovieSceneTools/Private/TrackEditors/CameraCutTrackEditor.cpp
+++ b/Engine/Source/Editor/MovieSceneTools/Private/TrackEditors/CameraCutTrackEditor.cpp
@@ -482,7 +482,7 @@ void FCameraCutTrackEditor::OnLockCameraClicked(ECheckBoxState CheckBoxState)
 		for (int32 i = 0; i < GEditor->LevelViewportClients.Num(); ++i)
 		{		
 			FLevelEditorViewportClient* LevelVC = GEditor->LevelViewportClients[i];
-			if (LevelVC && LevelVC->IsPerspective() && LevelVC->AllowsCinematicPreview() && LevelVC->GetViewMode() != VMI_Unknown)
+			if (LevelVC && LevelVC->IsPerspective() && LevelVC->AllowsCinematicControl() && LevelVC->GetViewMode() != VMI_Unknown)
 			{
 				LevelVC->SetActorLock(nullptr);
 				LevelVC->bLockedCameraView = false;
diff --git a/Engine/Source/Editor/MovieSceneTools/Private/TrackEditors/CinematicShotTrackEditor.cpp b/Engine/Source/Editor/MovieSceneTools/Private/TrackEditors/CinematicShotTrackEditor.cpp
--- a/Engine/Source/Editor/MovieSceneTools/Private/TrackEditors/CinematicShotTrackEditor.cpp
+++ b/Engine/Source/Editor/MovieSceneTools/Private/TrackEditors/CinematicShotTrackEditor.cpp
@@ -36,7 +36,6 @@
 #include "MovieSceneTimeHelpers.h"
 #include "Framework/Notifications/NotificationManager.h"
 #include "Widgets/Notifications/SNotificationList.h"
-#include "MovieSceneCaptureSettings.h"
 
 #define LOCTEXT_NAMESPACE "FCinematicShotTrackEditor"
 
@@ -710,7 +709,7 @@ void FCinematicShotTrackEditor::OnLockShotsClicked(ECheckBoxState CheckBoxState)
 		for (int32 i = 0; i < GEditor->LevelViewportClients.Num(); ++i)
 		{		
 			FLevelEditorViewportClient* LevelVC = GEditor->LevelViewportClients[i];
-			if (LevelVC && LevelVC->IsPerspective() && LevelVC->AllowsCinematicPreview() && LevelVC->GetViewMode() != VMI_Unknown)
+			if (LevelVC && LevelVC->IsPerspective() && LevelVC->AllowsCinematicControl() && LevelVC->GetViewMode() != VMI_Unknown)
 			{
 				LevelVC->SetActorLock(nullptr);
 				LevelVC->bLockedCameraView = false;
diff --git a/Engine/Source/Editor/Sequencer/Private/LevelEditorSequencerIntegration.cpp b/Engine/Source/Editor/Sequencer/Private/LevelEditorSequencerIntegration.cpp
--- a/Engine/Source/Editor/Sequencer/Private/LevelEditorSequencerIntegration.cpp
+++ b/Engine/Source/Editor/Sequencer/Private/LevelEditorSequencerIntegration.cpp
@@ -920,7 +920,7 @@ void FLevelEditorSequencerIntegration::ActivateRealtimeViewports()
 		if (LevelVC)
 		{
 			// If there is a director group, set the perspective viewports to realtime automatically.
-			if (LevelVC->IsPerspective() && LevelVC->AllowsCinematicPreview())
+			if (LevelVC->IsPerspective() && LevelVC->AllowsCinematicControl())
 			{				
 				// Ensure Realtime is turned on and store the original setting so we can restore it later.
 				LevelVC->SetRealtime(true, true);
@@ -940,7 +940,7 @@ void FLevelEditorSequencerIntegration::RestoreRealtimeViewports()
 		if (LevelVC)
 		{
 			// Turn off realtime when exiting.
-			if( LevelVC->IsPerspective() && LevelVC->AllowsCinematicPreview() )
+			if( LevelVC->IsPerspective() && LevelVC->AllowsCinematicControl() )
 			{				
 				// Specify true so RestoreRealtime will allow us to disable Realtime if it was original disabled
 				LevelVC->RestoreRealtime(true);
diff --git a/Engine/Source/Editor/Sequencer/Private/Sequencer.cpp b/Engine/Source/Editor/Sequencer/Private/Sequencer.cpp
--- a/Engine/Source/Editor/Sequencer/Private/Sequencer.cpp
+++ b/Engine/Source/Editor/Sequencer/Private/Sequencer.cpp
@@ -846,7 +846,7 @@ void FSequencer::PopToSequenceInstance(FMovieSceneSequenceIDRef SequenceID)
 		for (int32 i = 0; i < GEditor->LevelViewportClients.Num(); ++i)
 		{		
 			FLevelEditorViewportClient* LevelVC = GEditor->LevelViewportClients[i];
-			if (LevelVC && LevelVC->IsPerspective() && LevelVC->AllowsCinematicPreview() && LevelVC->GetViewMode() != VMI_Unknown)
+			if (LevelVC && LevelVC->IsPerspective() && LevelVC->AllowsCinematicControl() && LevelVC->GetViewMode() != VMI_Unknown)
 			{
 				LevelVC->SetActorLock(nullptr);
 				LevelVC->bLockedCameraView = false;
@@ -2687,7 +2687,7 @@ void FSequencer::UpdateCameraCut(UObject* CameraObject, UObject* UnlockIfCameraO
 
 	for (FLevelEditorViewportClient* LevelVC : GEditor->LevelViewportClients)
 	{
-		if ((LevelVC == nullptr) || !LevelVC->IsPerspective() || !LevelVC->AllowsCinematicPreview())
+		if ((LevelVC == nullptr) || !LevelVC->IsPerspective() || !LevelVC->AllowsCinematicControl())
 		{
 			continue;
 		}
@@ -2718,7 +2718,7 @@ void FSequencer::SetViewportSettings(const TMap<FViewportClient*, EMovieSceneVie
 	{
 		if (LevelVC && LevelVC->IsPerspective())
 		{
-			if (LevelVC->AllowsCinematicPreview())
+			if (LevelVC->AllowsCinematicControl())
 			{
 				if (ViewportParamsMap.Contains(LevelVC))
 				{
@@ -2754,7 +2754,7 @@ void FSequencer::GetViewportSettings(TMap<FViewportClient*, EMovieSceneViewportP
 {
 	for (FLevelEditorViewportClient* LevelVC : GEditor->LevelViewportClients)
 	{
-		if (LevelVC && LevelVC->IsPerspective() && LevelVC->AllowsCinematicPreview())
+		if (LevelVC && LevelVC->IsPerspective() && LevelVC->AllowsCinematicControl())
 		{
 			EMovieSceneViewportParams ViewportParams;
 			ViewportParams.FadeAmount = LevelVC->FadeAmount;
@@ -2782,7 +2782,7 @@ void FSequencer::SetPlaybackStatus(EMovieScenePlayerStatus::Type InPlaybackStatu
 
 	for (FLevelEditorViewportClient* LevelVC : GEditor->LevelViewportClients)
 	{
-		if (LevelVC && LevelVC->IsPerspective() && LevelVC->AllowsCinematicPreview())
+		if (LevelVC && LevelVC->IsPerspective() && LevelVC->AllowsCinematicControl())
 		{
 			LevelVC->ViewState.GetReference()->SetSequencerState(bIsPaused);
 		}
@@ -4473,7 +4473,7 @@ void FSequencer::SaveCurrentMovieScene()
 	{
 		for (FLevelEditorViewportClient* LevelVC : GEditor->LevelViewportClients)
 		{
-			if ((LevelVC == nullptr) || !LevelVC->IsPerspective() || !LevelVC->AllowsCinematicPreview())
+			if ((LevelVC == nullptr) || !LevelVC->IsPerspective() || !LevelVC->AllowsCinematicControl())
 			{
 				continue;
 			}
diff --git a/Engine/Source/Editor/Sequencer/Private/SequencerEdMode.cpp b/Engine/Source/Editor/Sequencer/Private/SequencerEdMode.cpp
--- a/Engine/Source/Editor/Sequencer/Private/SequencerEdMode.cpp
+++ b/Engine/Source/Editor/Sequencer/Private/SequencerEdMode.cpp
@@ -107,7 +107,7 @@ void FSequencerEdMode::DrawHUD(FEditorViewportClient* ViewportClient,FViewport*
 {
 	FEdMode::DrawHUD(ViewportClient,Viewport,View,Canvas);
 
-	if( ViewportClient->AllowsCinematicPreview() )
+	if( ViewportClient->AllowsCinematicControl() )
 	{
 		// Get the size of the viewport
 		const int32 SizeX = Viewport->GetSizeXY().X;
diff --git a/Engine/Source/Editor/UnrealEd/Private/EditorModeInterpolation.cpp b/Engine/Source/Editor/UnrealEd/Private/EditorModeInterpolation.cpp
--- a/Engine/Source/Editor/UnrealEd/Private/EditorModeInterpolation.cpp
+++ b/Engine/Source/Editor/UnrealEd/Private/EditorModeInterpolation.cpp
@@ -108,7 +108,7 @@ void FEdModeInterpEdit::CamMoveNotify(FEditorViewportClient* ViewportClient)
 	if(!InterpEd)
 		return;
 
-	if( ViewportClient->AllowsCinematicPreview() )
+	if( ViewportClient->AllowsCinematicControl() )
 	{
 		InterpEd->CamMoved( ViewportClient->GetViewLocation(), ViewportClient->GetViewRotation() );
 	}
diff --git a/Engine/Source/Editor/UnrealEd/Private/EditorViewportClient.cpp b/Engine/Source/Editor/UnrealEd/Private/EditorViewportClient.cpp
--- a/Engine/Source/Editor/UnrealEd/Private/EditorViewportClient.cpp
+++ b/Engine/Source/Editor/UnrealEd/Private/EditorViewportClient.cpp
@@ -294,7 +294,7 @@ void InitViewOptionsArray()
 }
 
 FEditorViewportClient::FEditorViewportClient(FEditorModeTools* InModeTools, FPreviewScene* InPreviewScene, const TWeakPtr<SEditorViewport>& InEditorViewportWidget)
-	: bAllowCinematicPreview(false)
+	: bAllowCinematicControl(false)
 	, CameraSpeedSetting(4)
 	, CameraSpeedScalar(1.0f)
 	, ImmersiveDelegate()
@@ -3531,7 +3531,7 @@ void FEditorViewportClient::Draw(FViewport* InViewport, FCanvas* Canvas)
 
 	ViewFamily.EngineShowFlags = EngineShowFlags;
 
-	if( ModeTools->GetActiveMode( FBuiltinEditorModes::EM_InterpEdit ) == 0 || !AllowsCinematicPreview() )
+	if( ModeTools->GetActiveMode( FBuiltinEditorModes::EM_InterpEdit ) == 0 || !AllowsCinematicControl() )
 	{
 		if( !EngineShowFlags.Game )
 		{
diff --git a/Engine/Source/Editor/UnrealEd/Private/LevelEditorViewport.cpp b/Engine/Source/Editor/UnrealEd/Private/LevelEditorViewport.cpp
--- a/Engine/Source/Editor/UnrealEd/Private/LevelEditorViewport.cpp
+++ b/Engine/Source/Editor/UnrealEd/Private/LevelEditorViewport.cpp
@@ -2151,7 +2151,7 @@ void FLevelEditorViewportClient::Tick(float DeltaTime)
 void FLevelEditorViewportClient::UpdateViewForLockedActor(float DeltaTime)
 {
 	// We can't be locked to a matinee actor if this viewport doesn't allow matinee control
-	if ( !bAllowCinematicPreview && ActorLockedByMatinee.IsValid() )
+	if ( !bAllowCinematicControl && ActorLockedByMatinee.IsValid() )
 	{
 		ActorLockedByMatinee = nullptr;
 	}
@@ -4431,7 +4431,7 @@ void FLevelEditorViewportClient::CopyLayoutFromViewport( const FLevelEditorViewp
 	ViewportType = InViewport.ViewportType;
 	SetOrthoZoom( InViewport.GetOrthoZoom() );
 	ActorLockedToCamera = InViewport.ActorLockedToCamera;
-	bAllowCinematicPreview = InViewport.bAllowCinematicPreview;
+	bAllowCinematicControl = InViewport.bAllowCinematicControl;
 }
 
 
diff --git a/Engine/Source/Editor/UnrealEd/Public/EditorViewportClient.h b/Engine/Source/Editor/UnrealEd/Public/EditorViewportClient.h
--- a/Engine/Source/Editor/UnrealEd/Public/EditorViewportClient.h
+++ b/Engine/Source/Editor/UnrealEd/Public/EditorViewportClient.h
@@ -713,10 +713,10 @@ public:
 	virtual bool DropObjectsAtCoordinates(int32 MouseX, int32 MouseY, const TArray<UObject*>& DroppedObjects, TArray<AActor*>& OutNewActors, bool bOnlyDropOnTarget = false, bool bCreateDropPreview = false, bool bSelectActors = true, UActorFactory* FactoryToUse = NULL ) { return false; }
 
 	/** Returns true if the viewport is allowed to be possessed for previewing cinematic sequences or keyframe animations*/
-	bool AllowsCinematicPreview() const { return bAllowCinematicPreview; }
+	bool AllowsCinematicControl() const { return bAllowCinematicControl; }
 
 	/** Sets whether or not this viewport is allowed to be possessed by cinematic/scrubbing tools */
-	void SetAllowCinematicPreview( bool bInAllowCinematicPreview ) { bAllowCinematicPreview = bInAllowCinematicPreview; }
+	void SetAllowCinematicControl( bool bInAllowCinematicControl) { bAllowCinematicControl = bInAllowCinematicControl; }
 
 public:
 	/** True if the window is maximized or floating */
@@ -979,7 +979,7 @@ public:
 protected:
 
 	/** true if this window is allowed to be possessed by cinematic tools for previewing sequences in real-time */
-	bool bAllowCinematicPreview;
+	bool bAllowCinematicControl;
 
 	/** Camera speed setting */
 	int32 CameraSpeedSetting;
diff --git a/Engine/Source/Editor/VREditor/Public/VREditorMode.h b/Engine/Source/Editor/VREditor/Public/VREditorMode.h
--- a/Engine/Source/Editor/VREditor/Public/VREditorMode.h
+++ b/Engine/Source/Editor/VREditor/Public/VREditorMode.h
@@ -218,7 +218,7 @@ public:
 		float TransformGizmoScale;
 		EHMDTrackingOrigin::Type TrackingOrigin;
 		float WorldToMetersScale;
-		bool bCinematicPreviewViewport;
+		bool bCinematicControlViewport;
 		bool bKeyAllEnabled;
 		EAutoChangeMode AutoChangeMode;
 
@@ -237,7 +237,7 @@ public:
 			  TransformGizmoScale( 1.0f ),
 			  TrackingOrigin(EHMDTrackingOrigin::Eye),
 			  WorldToMetersScale(100.0f),
-			  bCinematicPreviewViewport(false),
+			  bCinematicControlViewport(false),
 			  bKeyAllEnabled(false),
 			  AutoChangeMode()
 		{
diff --git a/Engine/Source/Editor/VREditor/VREditorMode.cpp b/Engine/Source/Editor/VREditor/VREditorMode.cpp
--- a/Engine/Source/Editor/VREditor/VREditorMode.cpp
+++ b/Engine/Source/Editor/VREditor/VREditorMode.cpp
@@ -1033,8 +1033,8 @@ void UVREditorMode::StartViewport(TSharedPtr<SLevelViewport> Viewport)
 		VREditorViewportClient.EngineShowFlags.SetSelectionOutline(true);
 
 		// Change viewport settings to more VR-friendly sequencer settings
-		SavedEditorState.bCinematicPreviewViewport = VRViewportClient.AllowsCinematicPreview();
-		VRViewportClient.SetAllowCinematicPreview(false);
+		SavedEditorState.bCinematicControlViewport = VRViewportClient.AllowsCinematicControl();
+		VRViewportClient.SetAllowCinematicControl(false);
 		// Need to force fading and color scaling off in case we enter VR editing mode with a sequence open
 		VRViewportClient.bEnableFading = false;
 		VRViewportClient.bEnableColorScaling = false;
@@ -1082,7 +1082,7 @@ void UVREditorMode::CloseViewport( const bool bShouldDisableStereo )
 			VRViewportClient.bAlwaysShowModeWidgetAfterSelectionChanges = SavedEditorState.bAlwaysShowModeWidgetAfterSelectionChanges;
 			VRViewportClient.EngineShowFlags = SavedEditorState.ShowFlags;
 			VRViewportClient.SetGameView(SavedEditorState.bGameView);
-			VRViewportClient.SetAllowCinematicPreview(SavedEditorState.bCinematicPreviewViewport);
+			VRViewportClient.SetAllowCinematicControl(SavedEditorState.bCinematicControlViewport);
 			VRViewportClient.bEnableFading = true;
 			VRViewportClient.bEnableColorScaling = true;
 			VRViewportClient.Invalidate(true);
