Changes to make new Clang in Xcode 11 beta compile the engine witout warnings

diff --git a/Engine/Source/Developer/Profiler/Private/ProfilerDataSource.cpp b/Engine/Source/Developer/Profiler/Private/ProfilerDataSource.cpp
--- a/Engine/Source/Developer/Profiler/Private/ProfilerDataSource.cpp
+++ b/Engine/Source/Developer/Profiler/Private/ProfilerDataSource.cpp
@@ -220,7 +220,7 @@ FName FEventGraphConsts::FakeRoot = TEXT("FakeRoot");
 	FEventGraphSample
 -----------------------------------------------------------------------------*/
 
-FEventProperty FEventGraphSample::Properties[ EEventPropertyIndex::InvalidOrMax ] =
+FEventProperty FEventGraphSample::Properties[ (uint32)EEventPropertyIndex::InvalidOrMax ] =
 {
 	// Properties
 	FEventProperty( EEventPropertyIndex::StatName, TEXT( "StatName" ), STRUCT_OFFSET( FEventGraphSample, _StatName ), EEventPropertyFormatters::Name ),
@@ -266,39 +266,39 @@ void FEventGraphSample::InitializePropertyManagement()
 	{
 		NamedProperties = TMapBuilder<FName,const FEventProperty*>()
 			// Properties
-			.Add( TEXT( "StatName" ), &Properties[EEventPropertyIndex::StatName] )
-			.Add( TEXT( "InclusiveTimeMS" ), &Properties[EEventPropertyIndex::InclusiveTimeMS] )
-			.Add( TEXT( "InclusiveTimePct" ), &Properties[EEventPropertyIndex::InclusiveTimePct] )
-			.Add( TEXT( "ExclusiveTimeMS" ), &Properties[EEventPropertyIndex::ExclusiveTimeMS] )
-			.Add( TEXT( "ExclusiveTimePct" ), &Properties[EEventPropertyIndex::ExclusiveTimePct] )
-			.Add( TEXT( "NumCallsPerFrame" ), &Properties[EEventPropertyIndex::NumCallsPerFrame] )
+			.Add( TEXT( "StatName" ), &Properties[(uint32)EEventPropertyIndex::StatName] )
+			.Add( TEXT( "InclusiveTimeMS" ), &Properties[(uint32)EEventPropertyIndex::InclusiveTimeMS] )
+			.Add( TEXT( "InclusiveTimePct" ), &Properties[(uint32)EEventPropertyIndex::InclusiveTimePct] )
+			.Add( TEXT( "ExclusiveTimeMS" ), &Properties[(uint32)EEventPropertyIndex::ExclusiveTimeMS] )
+			.Add( TEXT( "ExclusiveTimePct" ), &Properties[(uint32)EEventPropertyIndex::ExclusiveTimePct] )
+			.Add( TEXT( "NumCallsPerFrame" ), &Properties[(uint32)EEventPropertyIndex::NumCallsPerFrame] )
 
 			// Special none property
-			.Add( NAME_None, &Properties[EEventPropertyIndex::None] )
+			.Add( NAME_None, &Properties[(uint32)EEventPropertyIndex::None] )
 
-			.Add( TEXT( "MinInclusiveTimeMS" ), &Properties[EEventPropertyIndex::MinInclusiveTimeMS] )
-			.Add( TEXT( "MaxInclusiveTimeMS" ), &Properties[EEventPropertyIndex::MaxInclusiveTimeMS] )
-			.Add( TEXT( "AvgInclusiveTimeMS" ), &Properties[EEventPropertyIndex::AvgInclusiveTimeMS] )	
+			.Add( TEXT( "MinInclusiveTimeMS" ), &Properties[(uint32)EEventPropertyIndex::MinInclusiveTimeMS] )
+			.Add( TEXT( "MaxInclusiveTimeMS" ), &Properties[(uint32)EEventPropertyIndex::MaxInclusiveTimeMS] )
+			.Add( TEXT( "AvgInclusiveTimeMS" ), &Properties[(uint32)EEventPropertyIndex::AvgInclusiveTimeMS] )
 
-			.Add( TEXT( "MinNumCallsPerFrame" ), &Properties[EEventPropertyIndex::MinNumCallsPerFrame] )
-			.Add( TEXT( "MaxNumCallsPerFrame" ), &Properties[EEventPropertyIndex::MaxNumCallsPerFrame] )
-			.Add( TEXT( "AvgNumCallsPerFrame" ), &Properties[EEventPropertyIndex::AvgNumCallsPerFrame] )
+			.Add( TEXT( "MinNumCallsPerFrame" ), &Properties[(uint32)EEventPropertyIndex::MinNumCallsPerFrame] )
+			.Add( TEXT( "MaxNumCallsPerFrame" ), &Properties[(uint32)EEventPropertyIndex::MaxNumCallsPerFrame] )
+			.Add( TEXT( "AvgNumCallsPerFrame" ), &Properties[(uint32)EEventPropertyIndex::AvgNumCallsPerFrame] )
 
-			.Add( TEXT( "ThreadName" ), &Properties[EEventPropertyIndex::ThreadName] )
-			.Add( TEXT( "ThreadDurationMS" ), &Properties[EEventPropertyIndex::ThreadDurationMS] )
-			.Add( TEXT( "FrameDurationMS" ), &Properties[EEventPropertyIndex::FrameDurationMS] )
-			.Add( TEXT( "ThreadPct" ), &Properties[EEventPropertyIndex::ThreadPct] )
-			.Add( TEXT( "FramePct" ), &Properties[EEventPropertyIndex::FramePct] )
-			.Add( TEXT( "ThreadToFramePct" ), &Properties[EEventPropertyIndex::ThreadToFramePct] )
-			.Add( TEXT( "GroupName" ), &Properties[EEventPropertyIndex::GroupName] )
+			.Add( TEXT( "ThreadName" ), &Properties[(uint32)EEventPropertyIndex::ThreadName] )
+			.Add( TEXT( "ThreadDurationMS" ), &Properties[(uint32)EEventPropertyIndex::ThreadDurationMS] )
+			.Add( TEXT( "FrameDurationMS" ), &Properties[(uint32)EEventPropertyIndex::FrameDurationMS] )
+			.Add( TEXT( "ThreadPct" ), &Properties[(uint32)EEventPropertyIndex::ThreadPct] )
+			.Add( TEXT( "FramePct" ), &Properties[(uint32)EEventPropertyIndex::FramePct] )
+			.Add( TEXT( "ThreadToFramePct" ), &Properties[(uint32)EEventPropertyIndex::ThreadToFramePct] )
+			.Add( TEXT( "GroupName" ), &Properties[(uint32)EEventPropertyIndex::GroupName] )
 
 			// Booleans
-			.Add( TEXT( "bIsHotPath" ), &Properties[EEventPropertyIndex::bIsHotPath] )
-			.Add( TEXT( "bIsFiltered" ), &Properties[EEventPropertyIndex::bIsFiltered] )
-			.Add( TEXT( "bIsCulled" ), &Properties[EEventPropertyIndex::bIsCulled] )
+			.Add( TEXT( "bIsHotPath" ), &Properties[(uint32)EEventPropertyIndex::bIsHotPath] )
+			.Add( TEXT( "bIsFiltered" ), &Properties[(uint32)EEventPropertyIndex::bIsFiltered] )
+			.Add( TEXT( "bIsCulled" ), &Properties[(uint32)EEventPropertyIndex::bIsCulled] )
 
 			// Booleans internal
-			.Add( TEXT( "bNeedNotCulledChildrenUpdate" ), &Properties[EEventPropertyIndex::bNeedNotCulledChildrenUpdate] )
+			.Add( TEXT( "bNeedNotCulledChildrenUpdate" ), &Properties[(uint32)EEventPropertyIndex::bNeedNotCulledChildrenUpdate] )
 			;
 	
 		// Make sure that the minimal property manager has been initialized.
@@ -308,8 +308,8 @@ void FEventGraphSample::InitializePropertyManagement()
 		check( NamedProperties.FindChecked( NAME_None )->Name == NAME_None );
 		check( NamedProperties.FindChecked( NAME_None )->Offset == INDEX_NONE );
 
-		check( FEventGraphSample::Properties[ EEventPropertyIndex::None ].Name == NAME_None );
-		check( FEventGraphSample::Properties[ EEventPropertyIndex::None ].Offset == INDEX_NONE );
+		check( FEventGraphSample::Properties[ (uint32)EEventPropertyIndex::None ].Name == NAME_None );
+		check( FEventGraphSample::Properties[ (uint32)EEventPropertyIndex::None ].Offset == INDEX_NONE );
 	}
 }
 
diff --git a/Engine/Source/Developer/Profiler/Private/ProfilerDataSource.h b/Engine/Source/Developer/Profiler/Private/ProfilerDataSource.h
--- a/Engine/Source/Developer/Profiler/Private/ProfilerDataSource.h
+++ b/Engine/Source/Developer/Profiler/Private/ProfilerDataSource.h
@@ -690,7 +690,7 @@ protected:
 -----------------------------------------------------------------------------*/
 
 /** Enumerates event graph columns index. */
-enum EEventPropertyIndex
+enum class EEventPropertyIndex : uint32
 {
 	/** Stat name must be the first column, because of the expander arrow. */
 	StatName,
@@ -968,7 +968,7 @@ public:
 
 protected:
 	/** Contains all properties of the event graph sample class. */
-	static FEventProperty Properties[ EEventPropertyIndex::InvalidOrMax ];
+	static FEventProperty Properties[ (uint32)EEventPropertyIndex::InvalidOrMax ];
 
 	/** Contains all properties of the event graph sample class, indexed by the name of the property, stored as FName -> FEventProperty&. */
 	static TMap<FName,const FEventProperty*> NamedProperties;
diff --git a/Engine/Source/Developer/Profiler/Private/Widgets/SEventGraph.cpp b/Engine/Source/Developer/Profiler/Private/Widgets/SEventGraph.cpp
--- a/Engine/Source/Developer/Profiler/Private/Widgets/SEventGraph.cpp
+++ b/Engine/Source/Developer/Profiler/Private/Widgets/SEventGraph.cpp
@@ -73,14 +73,14 @@ struct FEventGraphColumns
 {
 	/** Default constructor. */
 	FEventGraphColumns()
-		: NumColumns( EEventPropertyIndex::None+1 )
+		: NumColumns( (uint32)EEventPropertyIndex::None+1 )
 	{
 		// Make event property is initialized.
 		FEventGraphSample::InitializePropertyManagement();
 
 		Collection = new FEventGraphColumn[NumColumns];
 
-		Collection[EEventPropertyIndex::StatName] = FEventGraphColumn
+		Collection[(uint32)EEventPropertyIndex::StatName] = FEventGraphColumn
 		(
 			EEventPropertyIndex::StatName,
 			TEXT( "name" ),
@@ -91,7 +91,7 @@ struct FEventGraphColumns
 			0.0f
 		);
 
-		Collection[EEventPropertyIndex::InclusiveTimeMS] = FEventGraphColumn
+		Collection[(uint32)EEventPropertyIndex::InclusiveTimeMS] = FEventGraphColumn
 		(
 			EEventPropertyIndex::InclusiveTimeMS,
 			TEXT( "inc" ),
@@ -102,7 +102,7 @@ struct FEventGraphColumns
 			72.0f
 		);
 
-		Collection[EEventPropertyIndex::InclusiveTimePct] = FEventGraphColumn
+		Collection[(uint32)EEventPropertyIndex::InclusiveTimePct] = FEventGraphColumn
 		(
 			EEventPropertyIndex::InclusiveTimePct,
 			TEXT( "inc%" ),
@@ -113,7 +113,7 @@ struct FEventGraphColumns
 			72.0f
 		);
 
-		Collection[EEventPropertyIndex::ExclusiveTimeMS] = FEventGraphColumn
+		Collection[(uint32)EEventPropertyIndex::ExclusiveTimeMS] = FEventGraphColumn
 		(
 			EEventPropertyIndex::ExclusiveTimeMS,
 			TEXT( "exc" ),
@@ -124,7 +124,7 @@ struct FEventGraphColumns
 			72.0f
 		);
 
-		Collection[EEventPropertyIndex::ExclusiveTimePct] = FEventGraphColumn
+		Collection[(uint32)EEventPropertyIndex::ExclusiveTimePct] = FEventGraphColumn
 		(
 			EEventPropertyIndex::ExclusiveTimePct,
 			TEXT( "exc%" ),
@@ -135,7 +135,7 @@ struct FEventGraphColumns
 			72.0f
 		);
 
-		Collection[EEventPropertyIndex::NumCallsPerFrame] = FEventGraphColumn
+		Collection[(uint32)EEventPropertyIndex::NumCallsPerFrame] = FEventGraphColumn
 		(
 			EEventPropertyIndex::NumCallsPerFrame,
 			TEXT( "calls" ),
@@ -147,7 +147,7 @@ struct FEventGraphColumns
 		);
 
 		// Fake column used as a default column for NAME_None
-		Collection[EEventPropertyIndex::None] = FEventGraphColumn
+		Collection[(uint32)EEventPropertyIndex::None] = FEventGraphColumn
 		(
 			EEventPropertyIndex::None,
 			TEXT( "None" ),
@@ -159,13 +159,13 @@ struct FEventGraphColumns
 		);
 
 		ColumnNameToIndexMapping = TMapBuilder<FName, const FEventGraphColumn*>()
-			.Add( TEXT( "StatName" ), &Collection[EEventPropertyIndex::StatName] )
-			.Add( TEXT( "InclusiveTimeMS" ), &Collection[EEventPropertyIndex::InclusiveTimeMS] )
-			.Add( TEXT( "InclusiveTimePct" ), &Collection[EEventPropertyIndex::InclusiveTimePct] )
-			.Add( TEXT( "ExclusiveTimeMS" ), &Collection[EEventPropertyIndex::ExclusiveTimeMS] )
-			.Add( TEXT( "ExclusiveTimePct" ), &Collection[EEventPropertyIndex::ExclusiveTimePct] )
-			.Add( TEXT( "NumCallsPerFrame" ), &Collection[EEventPropertyIndex::NumCallsPerFrame] )
-			.Add( NAME_None, &Collection[EEventPropertyIndex::None] )
+			.Add( TEXT( "StatName" ), &Collection[(uint32)EEventPropertyIndex::StatName] )
+			.Add( TEXT( "InclusiveTimeMS" ), &Collection[(uint32)EEventPropertyIndex::InclusiveTimeMS] )
+			.Add( TEXT( "InclusiveTimePct" ), &Collection[(uint32)EEventPropertyIndex::InclusiveTimePct] )
+			.Add( TEXT( "ExclusiveTimeMS" ), &Collection[(uint32)EEventPropertyIndex::ExclusiveTimeMS] )
+			.Add( TEXT( "ExclusiveTimePct" ), &Collection[(uint32)EEventPropertyIndex::ExclusiveTimePct] )
+			.Add( TEXT( "NumCallsPerFrame" ), &Collection[(uint32)EEventPropertyIndex::NumCallsPerFrame] )
+			.Add( NAME_None, &Collection[(uint32)EEventPropertyIndex::None] )
 			;
 	}
 
@@ -495,7 +495,7 @@ public:
 			SNew( SEventGraphTableCell, SharedThis(this), OwnerEventGraph )
 			.Visibility( this, &SEventGraphTableRow::IsColumnVisible, ColumnID )
 			.ColumnID( ColumnID )
-			.IsEventNameColumn( ColumnID == FEventGraphColumns::Get().Collection[EEventPropertyIndex::StatName].ID )
+			.IsEventNameColumn( ColumnID == FEventGraphColumns::Get().Collection[(uint32)EEventPropertyIndex::StatName].ID )
 			.EventPtr( EventPtr )
 			.OnSetHoveredTableCell( this, &SEventGraphTableRow::OnSetHoveredTableCell )
 		];
@@ -2106,7 +2106,7 @@ void SEventGraph::TreeViewHeaderRow_CreateColumnArgs( const uint32 ColumnIndex )
 void SEventGraph::InitializeAndShowHeaderColumns()
 {
 	ColumnSortMode = EColumnSortMode::Descending;
-	ColumnBeingSorted = FEventGraphColumns::Get().Collection[EEventPropertyIndex::InclusiveTimeMS].ID;
+	ColumnBeingSorted = FEventGraphColumns::Get().Collection[(uint32)EEventPropertyIndex::InclusiveTimeMS].ID;
 
 	for (uint32 ColumnIndex = 0; ColumnIndex < FEventGraphColumns::Get().NumColumns; ColumnIndex++)
 	{
@@ -2287,7 +2287,7 @@ void SEventGraph::ContextMenu_ExpandHotPath_Execute()
 	FEventGraphSamplePtr EventPtr = SelectedItems[0];
 
 	ColumnSortMode = EColumnSortMode::Descending;
-	ColumnBeingSorted = FEventGraphColumns::Get().Collection[EEventPropertyIndex::InclusiveTimeMS].ID;
+	ColumnBeingSorted = FEventGraphColumns::Get().Collection[(uint32)EEventPropertyIndex::InclusiveTimeMS].ID;
 	SortEvents();
 	
 	// Clear hot path
@@ -2424,7 +2424,7 @@ bool SEventGraph::ContextMenu_SortMode_IsChecked( const EColumnSortMode::Type In
 void SEventGraph::ContextMenu_ResetColumns_Execute()
 {
 	ColumnSortMode = EColumnSortMode::Descending;
-	ColumnBeingSorted = FEventGraphColumns::Get().Collection[EEventPropertyIndex::InclusiveTimeMS].ID;
+	ColumnBeingSorted = FEventGraphColumns::Get().Collection[(uint32)EEventPropertyIndex::InclusiveTimeMS].ID;
 
 	for( uint32 ColumnIndex = 0; ColumnIndex < FEventGraphColumns::Get().NumColumns; ColumnIndex++ )
 	{
@@ -2553,12 +2553,12 @@ void SEventGraph::ShowEventsInViewMode( const TArray<FEventGraphSamplePtr>& Even
 	if( NewViewMode == EEventGraphViewModes::FlatInclusive || NewViewMode == EEventGraphViewModes::FlatInclusiveCoalesced || NewViewMode == EEventGraphViewModes::Hierarchical )
 	{
 		ColumnIndex = EEventPropertyIndex::InclusiveTimeMS;
-		SetSortModeForColumn( FEventGraphColumns::Get().Collection[ColumnIndex].ID, EColumnSortMode::Descending );
+		SetSortModeForColumn( FEventGraphColumns::Get().Collection[(uint32)ColumnIndex].ID, EColumnSortMode::Descending );
 	}
 	else if( NewViewMode == EEventGraphViewModes::FlatExclusive || NewViewMode == EEventGraphViewModes::FlatExclusiveCoalesced )
 	{
 		ColumnIndex = EEventPropertyIndex::ExclusiveTimeMS;
-		SetSortModeForColumn( FEventGraphColumns::Get().Collection[ColumnIndex].ID, EColumnSortMode::Descending );
+		SetSortModeForColumn( FEventGraphColumns::Get().Collection[(uint32)ColumnIndex].ID, EColumnSortMode::Descending );
 	}
 
 	ScrollToTheSlowestSelectedEvent( ColumnIndex );
@@ -2589,7 +2589,7 @@ void SEventGraph::ScrollToTheSlowestSelectedEvent( EEventPropertyIndex ColumnInd
 	if( SelectedEvents.Num() > 0 )
 	{
 		// Sort events by the inclusive or the exclusive time, depends on the view mode.
-		const FEventGraphColumn& Column = FEventGraphColumns::Get().Collection[ColumnIndex];
+		const FEventGraphColumn& Column = FEventGraphColumns::Get().Collection[(uint32)ColumnIndex];
 		FEventArraySorter::Sort( SelectedEvents, Column.ID, EEventCompareOps::Greater );
 
 		// Scroll to the the slowest item.
@@ -2804,7 +2804,7 @@ void SEventGraph::GenerateTopEvents( const TSet< FEventGraphSamplePtr >& EventPt
 	}
 
 	// Sort events by the inclusive time.
-	const FEventGraphColumn& Column = FEventGraphColumns::Get().Collection[EEventPropertyIndex::InclusiveTimeMS];
+	const FEventGraphColumn& Column = FEventGraphColumns::Get().Collection[(uint32)EEventPropertyIndex::InclusiveTimeMS];
 	FEventArraySorter::Sort( EventPtrArray, Column.ID, EEventCompareOps::Greater );
 
 	// Calculate total time for the top events.
diff --git a/Engine/Source/Editor/MaterialEditor/Private/MaterialEditor.cpp b/Engine/Source/Editor/MaterialEditor/Private/MaterialEditor.cpp
--- a/Engine/Source/Editor/MaterialEditor/Private/MaterialEditor.cpp
+++ b/Engine/Source/Editor/MaterialEditor/Private/MaterialEditor.cpp
@@ -2036,7 +2036,7 @@ void FMaterialEditor::UpdateMaterialinfoList_Old()
 			}
 
 			// Only show general info if there are no errors and stats are enabled - Stats show for Materials, layers and blends
-		if (CompileErrors.Num() == 0 && (!MaterialFunction || MaterialFunction->GetMaterialFunctionUsage() != Default))
+		if (CompileErrors.Num() == 0 && (!MaterialFunction || MaterialFunction->GetMaterialFunctionUsage() != EMaterialFunctionUsage::Default))
 			{
 			TArray<FMaterialStatsUtils::FShaderInstructionsInfo> Results;
 			TArray<FMaterialStatsUtils::FShaderInstructionsInfo> EmptyMaterialResults;
diff --git a/Engine/Source/Editor/Matinee/Private/Matinee.cpp b/Engine/Source/Editor/Matinee/Private/Matinee.cpp
--- a/Engine/Source/Editor/Matinee/Private/Matinee.cpp
+++ b/Engine/Source/Editor/Matinee/Private/Matinee.cpp
@@ -1205,15 +1205,15 @@ void FMatinee::BindCommands()
 	ToolkitCommands->MapAction( Commands.ExportAnimGroupFBX, FExecuteAction::CreateSP(this, &FMatinee::OnContextGroupExportAnimFBX) );
 	ToolkitCommands->MapAction( 
 		Commands.GroupDuplicate, 
-		FExecuteAction::CreateSP(this, &FMatinee::OnContextNewGroup, FMatineeCommands::EGroupAction::DuplicateGroup),
-		FCanExecuteAction::CreateSP(this, &FMatinee::CanCreateNewGroup, FMatineeCommands::EGroupAction::DuplicateGroup)
+		FExecuteAction::CreateSP(this, &FMatinee::OnContextNewGroup, FMatineeCommands::EGroupAction::Type::DuplicateGroup),
+		FCanExecuteAction::CreateSP(this, &FMatinee::CanCreateNewGroup, FMatineeCommands::EGroupAction::Type::DuplicateGroup)
 		);
 	ToolkitCommands->MapAction( Commands.GroupDelete, FExecuteAction::CreateSP(this, &FMatinee::OnContextGroupDelete), FCanExecuteAction::CreateSP(this, &FMatinee::CanGroupDelete) );
 	ToolkitCommands->MapAction( Commands.GroupCreateTab, FExecuteAction::CreateSP(this, &FMatinee::OnContextGroupCreateTab), FCanExecuteAction::CreateSP(this, &FMatinee::CanGroupCreateTab) );
 	ToolkitCommands->MapAction( Commands.GroupRemoveFromTab, FExecuteAction::CreateSP(this, &FMatinee::OnContextGroupRemoveFromTab) );
 	ToolkitCommands->MapAction( 
 		Commands.RemoveFromGroupFolder, 
-		FExecuteAction::CreateSP(this, &FMatinee::OnContextGroupChangeGroupFolder, FMatineeCommands::EGroupAction::RemoveFromGroupFolder, -1) 
+		FExecuteAction::CreateSP(this, &FMatinee::OnContextGroupChangeGroupFolder, FMatineeCommands::EGroupAction::Type::RemoveFromGroupFolder, -1)
 		);
 
 	//Track Context Menu
@@ -1228,13 +1228,13 @@ void FMatinee::BindCommands()
 	ToolkitCommands->MapAction( Commands.ExportAnimTrackFBX, FExecuteAction::CreateSP(this, &FMatinee::OnContextTrackExportAnimFBX) );
 
 	//Background Context Menu
-	ToolkitCommands->MapAction( Commands.NewFolder, FExecuteAction::CreateSP(this, &FMatinee::OnContextNewGroup, FMatineeCommands::EGroupAction::NewFolder) );
-	ToolkitCommands->MapAction( Commands.NewEmptyGroup, FExecuteAction::CreateSP(this, &FMatinee::OnContextNewGroup, FMatineeCommands::EGroupAction::NewEmptyGroup) );
-	ToolkitCommands->MapAction( Commands.NewCameraGroup, FExecuteAction::CreateSP(this, &FMatinee::OnContextNewGroup, FMatineeCommands::EGroupAction::NewCameraGroup) );
-	ToolkitCommands->MapAction( Commands.NewParticleGroup, FExecuteAction::CreateSP(this, &FMatinee::OnContextNewGroup, FMatineeCommands::EGroupAction::NewParticleGroup) );
-	ToolkitCommands->MapAction( Commands.NewSkeletalMeshGroup, FExecuteAction::CreateSP(this, &FMatinee::OnContextNewGroup, FMatineeCommands::EGroupAction::NewSkeletalMeshGroup) );
-	ToolkitCommands->MapAction( Commands.NewLightingGroup, FExecuteAction::CreateSP(this, &FMatinee::OnContextNewGroup, FMatineeCommands::EGroupAction::NewLightingGroup) );
-	ToolkitCommands->MapAction( Commands.NewDirectorGroup, FExecuteAction::CreateSP(this, &FMatinee::OnContextNewGroup, FMatineeCommands::EGroupAction::NewDirectorGroup) );
+	ToolkitCommands->MapAction( Commands.NewFolder, FExecuteAction::CreateSP(this, &FMatinee::OnContextNewGroup, FMatineeCommands::EGroupAction::Type::NewFolder) );
+	ToolkitCommands->MapAction( Commands.NewEmptyGroup, FExecuteAction::CreateSP(this, &FMatinee::OnContextNewGroup, FMatineeCommands::EGroupAction::Type::NewEmptyGroup) );
+	ToolkitCommands->MapAction( Commands.NewCameraGroup, FExecuteAction::CreateSP(this, &FMatinee::OnContextNewGroup, FMatineeCommands::EGroupAction::Type::NewCameraGroup) );
+	ToolkitCommands->MapAction( Commands.NewParticleGroup, FExecuteAction::CreateSP(this, &FMatinee::OnContextNewGroup, FMatineeCommands::EGroupAction::Type::NewParticleGroup) );
+	ToolkitCommands->MapAction( Commands.NewSkeletalMeshGroup, FExecuteAction::CreateSP(this, &FMatinee::OnContextNewGroup, FMatineeCommands::EGroupAction::Type::NewSkeletalMeshGroup) );
+	ToolkitCommands->MapAction( Commands.NewLightingGroup, FExecuteAction::CreateSP(this, &FMatinee::OnContextNewGroup, FMatineeCommands::EGroupAction::Type::NewLightingGroup) );
+	ToolkitCommands->MapAction( Commands.NewDirectorGroup, FExecuteAction::CreateSP(this, &FMatinee::OnContextNewGroup, FMatineeCommands::EGroupAction::Type::NewDirectorGroup) );
 
 	// Menu
 	ToolkitCommands->MapAction(Commands.ToggleCurveEditor,
@@ -1249,11 +1249,11 @@ void FMatinee::BindCommands()
 		);
 
 	//Key Context Menu
-	ToolkitCommands->MapAction( Commands.KeyModeCurveAuto, FExecuteAction::CreateSP(this, &FMatinee::OnContextKeyInterpMode, FMatineeCommands::EKeyAction::KeyModeCurveAuto) );
-	ToolkitCommands->MapAction( Commands.KeyModeCurveAutoClamped, FExecuteAction::CreateSP(this, &FMatinee::OnContextKeyInterpMode, FMatineeCommands::EKeyAction::KeyModeCurveAutoClamped) );
-	ToolkitCommands->MapAction( Commands.KeyModeCurveBreak, FExecuteAction::CreateSP(this, &FMatinee::OnContextKeyInterpMode, FMatineeCommands::EKeyAction::KeyModeCurveBreak) );
-	ToolkitCommands->MapAction( Commands.KeyModeLinear, FExecuteAction::CreateSP(this, &FMatinee::OnContextKeyInterpMode, FMatineeCommands::EKeyAction::KeyModeLinear) );
-	ToolkitCommands->MapAction( Commands.KeyModeConstant, FExecuteAction::CreateSP(this, &FMatinee::OnContextKeyInterpMode, FMatineeCommands::EKeyAction::KeyModeConstant) );
+	ToolkitCommands->MapAction( Commands.KeyModeCurveAuto, FExecuteAction::CreateSP(this, &FMatinee::OnContextKeyInterpMode, FMatineeCommands::EKeyAction::Type::KeyModeCurveAuto) );
+	ToolkitCommands->MapAction( Commands.KeyModeCurveAutoClamped, FExecuteAction::CreateSP(this, &FMatinee::OnContextKeyInterpMode, FMatineeCommands::EKeyAction::Type::KeyModeCurveAutoClamped) );
+	ToolkitCommands->MapAction( Commands.KeyModeCurveBreak, FExecuteAction::CreateSP(this, &FMatinee::OnContextKeyInterpMode, FMatineeCommands::EKeyAction::Type::KeyModeCurveBreak) );
+	ToolkitCommands->MapAction( Commands.KeyModeLinear, FExecuteAction::CreateSP(this, &FMatinee::OnContextKeyInterpMode, FMatineeCommands::EKeyAction::Type::KeyModeLinear) );
+	ToolkitCommands->MapAction( Commands.KeyModeConstant, FExecuteAction::CreateSP(this, &FMatinee::OnContextKeyInterpMode, FMatineeCommands::EKeyAction::Type::KeyModeConstant) );
 	ToolkitCommands->MapAction( Commands.KeySetTime, FExecuteAction::CreateSP(this, &FMatinee::OnContextSetKeyTime) );
 	ToolkitCommands->MapAction( Commands.KeySetValue, FExecuteAction::CreateSP(this, &FMatinee::OnContextSetValue) );
 	ToolkitCommands->MapAction( Commands.KeySetBool, FExecuteAction::CreateSP(this, &FMatinee::OnContextSetBool) );
@@ -1266,19 +1266,19 @@ void FMatinee::BindCommands()
 	ToolkitCommands->MapAction( Commands.ToggleKeyFlip, FExecuteAction::CreateSP(this, &FMatinee::OnFlipToggleKey) );
 
 	ToolkitCommands->MapAction( Commands.KeySetConditionAlways, 
-		FExecuteAction::CreateSP(this, &FMatinee::OnKeyContext_SetCondition, FMatineeCommands::EKeyAction::ConditionAlways),
+		FExecuteAction::CreateSP(this, &FMatinee::OnKeyContext_SetCondition, FMatineeCommands::EKeyAction::Type::ConditionAlways),
 		FCanExecuteAction(),
-		FIsActionChecked::CreateSP(this, &FMatinee::KeyContext_IsSetConditionToggled, FMatineeCommands::EKeyAction::ConditionAlways)
+		FIsActionChecked::CreateSP(this, &FMatinee::KeyContext_IsSetConditionToggled, FMatineeCommands::EKeyAction::Type::ConditionAlways)
 		);
 	ToolkitCommands->MapAction( Commands.KeySetConditionGoreEnabled, 
-		FExecuteAction::CreateSP(this, &FMatinee::OnKeyContext_SetCondition, FMatineeCommands::EKeyAction::ConditionGoreEnabled),
+		FExecuteAction::CreateSP(this, &FMatinee::OnKeyContext_SetCondition, FMatineeCommands::EKeyAction::Type::ConditionGoreEnabled),
 		FCanExecuteAction(),
-		FIsActionChecked::CreateSP(this, &FMatinee::KeyContext_IsSetConditionToggled, FMatineeCommands::EKeyAction::ConditionGoreEnabled)
+		FIsActionChecked::CreateSP(this, &FMatinee::KeyContext_IsSetConditionToggled, FMatineeCommands::EKeyAction::Type::ConditionGoreEnabled)
 		);
 	ToolkitCommands->MapAction( Commands.KeySetConditionGoreDisabled, 
-		FExecuteAction::CreateSP(this, &FMatinee::OnKeyContext_SetCondition, FMatineeCommands::EKeyAction::ConditionGoreDisabled),
+		FExecuteAction::CreateSP(this, &FMatinee::OnKeyContext_SetCondition, FMatineeCommands::EKeyAction::Type::ConditionGoreDisabled),
 		FCanExecuteAction(),
-		FIsActionChecked::CreateSP(this, &FMatinee::KeyContext_IsSetConditionToggled, FMatineeCommands::EKeyAction::ConditionGoreDisabled)
+		FIsActionChecked::CreateSP(this, &FMatinee::KeyContext_IsSetConditionToggled, FMatineeCommands::EKeyAction::Type::ConditionGoreDisabled)
 		);
 
 	ToolkitCommands->MapAction( Commands.AnimKeyLoop, FExecuteAction::CreateSP(this, &FMatinee::OnSetAnimKeyLooping, true) );
diff --git a/Engine/Source/Editor/Matinee/Private/MatineeActions.h b/Engine/Source/Editor/Matinee/Private/MatineeActions.h
--- a/Engine/Source/Editor/Matinee/Private/MatineeActions.h
+++ b/Engine/Source/Editor/Matinee/Private/MatineeActions.h
@@ -195,7 +195,7 @@ public:
 	/** Action Types */
 	struct EGroupAction
 	{
-		enum Type
+		enum class Type
 		{
 			NewFolder, 
 			NewEmptyGroup,
@@ -213,7 +213,7 @@ public:
 
 	struct EKeyAction
 	{
-		enum Type
+		enum class Type
 		{
 			KeyModeCurveAuto,
 			KeyModeCurveAutoClamped, 
diff --git a/Engine/Source/Editor/Matinee/Private/MatineeMenus.cpp b/Engine/Source/Editor/Matinee/Private/MatineeMenus.cpp
--- a/Engine/Source/Editor/Matinee/Private/MatineeMenus.cpp
+++ b/Engine/Source/Editor/Matinee/Private/MatineeMenus.cpp
@@ -245,10 +245,10 @@ void FMatinee::OnMenuAddKey()
 
 void FMatinee::OnContextNewGroup( FMatineeCommands::EGroupAction::Type InActionId)
 {
-	const bool bIsNewFolder = (InActionId == FMatineeCommands::EGroupAction::NewFolder); 
-	const bool bDirGroup = (InActionId == FMatineeCommands::EGroupAction::NewDirectorGroup); 
-	const bool bDuplicateGroup = (InActionId == FMatineeCommands::EGroupAction::DuplicateGroup); 
-	const bool bLightingGroup = (InActionId == FMatineeCommands::EGroupAction::NewLightingGroup); 
+	const bool bIsNewFolder = (InActionId == FMatineeCommands::EGroupAction::Type::NewFolder);
+	const bool bDirGroup = (InActionId == FMatineeCommands::EGroupAction::Type::NewDirectorGroup);
+	const bool bDuplicateGroup = (InActionId == FMatineeCommands::EGroupAction::Type::DuplicateGroup);
+	const bool bLightingGroup = (InActionId == FMatineeCommands::EGroupAction::Type::NewLightingGroup);
 
 	// Only one director group is allowed
 	if (bDirGroup && IData->FindDirectorGroup())
@@ -333,9 +333,9 @@ bool FMatinee::CanCreateNewGroup( FMatineeCommands::EGroupAction::Type InActionI
 void FMatinee::NewGroupPopup( FMatineeCommands::EGroupAction::Type InActionId, AActor* GroupActor, TArray<AActor *> OtherActorsToAddToGroup )
 {
 	// Find out if we want to make a 'Director' group.
-	const bool bIsNewFolder = (InActionId == FMatineeCommands::EGroupAction::NewFolder); 
-	const bool bDirGroup = (InActionId == FMatineeCommands::EGroupAction::NewDirectorGroup); 
-	const bool bDuplicateGroup = (InActionId == FMatineeCommands::EGroupAction::DuplicateGroup); 
+	const bool bIsNewFolder = (InActionId == FMatineeCommands::EGroupAction::Type::NewFolder);
+	const bool bDirGroup = (InActionId == FMatineeCommands::EGroupAction::Type::NewDirectorGroup);
+	const bool bDuplicateGroup = (InActionId == FMatineeCommands::EGroupAction::Type::DuplicateGroup);
 
 	// If not a director group - ask for a name.
 	if(!bDirGroup)
@@ -344,32 +344,32 @@ void FMatinee::NewGroupPopup( FMatineeCommands::EGroupAction::Type InActionId, A
 		FText DefaultNewGroupName;
 		switch( InActionId )
 		{
-		case FMatineeCommands::EGroupAction::NewCameraGroup: 
+		case FMatineeCommands::EGroupAction::Type::NewCameraGroup:
 			DialogName = LOCTEXT( "NewGroupName", "New Group Name" );
 			DefaultNewGroupName = LOCTEXT( "NewCameraGroup", "NewCameraGroup" );
 			break;
 
-		case FMatineeCommands::EGroupAction::NewParticleGroup: 
+		case FMatineeCommands::EGroupAction::Type::NewParticleGroup:
 			DialogName = LOCTEXT( "NewGroupName", "New Group Name" );
 			DefaultNewGroupName = LOCTEXT( "NewParticleGroup", "NewParticleGroup" );
 			break;
 
-		case FMatineeCommands::EGroupAction::NewSkeletalMeshGroup: 
+		case FMatineeCommands::EGroupAction::Type::NewSkeletalMeshGroup:
 			DialogName = LOCTEXT( "NewGroupName", "New Group Name" );
 			DefaultNewGroupName = LOCTEXT( "NewSkeletalMeshGroup", "NewSkeletalMeshGroup" );
 			break;
 
-		case FMatineeCommands::EGroupAction::NewLightingGroup: 
+		case FMatineeCommands::EGroupAction::Type::NewLightingGroup:
 			DialogName = LOCTEXT( "NewGroupName", "New Group Name" );
 			DefaultNewGroupName = LOCTEXT( "NewLightingGroup", "NewLightingGroup" );
 			break;
 
-		case FMatineeCommands::EGroupAction::NewFolder: 
+		case FMatineeCommands::EGroupAction::Type::NewFolder:
 			DialogName = LOCTEXT( "NewFolderName", "New Folder Name" );
 			DefaultNewGroupName = LOCTEXT( "NewFolder", "NewFolder" );
 			break;
 
-		case FMatineeCommands::EGroupAction::DuplicateGroup: 
+		case FMatineeCommands::EGroupAction::Type::DuplicateGroup:
 			// When duplicating, we use unlocalized text 
 			// at the moment. So, the spaces are needed.
 			DialogName = LOCTEXT( "NewGroupName", "New Group Name" );
@@ -421,9 +421,9 @@ void FMatinee::NewGroupPopupTextCommitted(
 	//note: we don't need to check commit type... handled by GetNewNamePopup
 
 	FName NewGroupName = FName(*InText.ToString().Left(NAME_SIZE));
-	const bool bIsNewFolder = (InActionId == FMatineeCommands::EGroupAction::NewFolder); 
-	const bool bDirGroup = (InActionId == FMatineeCommands::EGroupAction::NewDirectorGroup); 
-	const bool bDuplicateGroup = (InActionId == FMatineeCommands::EGroupAction::DuplicateGroup); 
+	const bool bIsNewFolder = (InActionId == FMatineeCommands::EGroupAction::Type::NewFolder);
+	const bool bDirGroup = (InActionId == FMatineeCommands::EGroupAction::Type::NewDirectorGroup);
+	const bool bDuplicateGroup = (InActionId == FMatineeCommands::EGroupAction::Type::DuplicateGroup);
 
 	TMap<UInterpGroup*, FName> DuplicateGroupToNameMap;
 	if (bDuplicateGroup && GroupToDuplicate != NULL)
@@ -485,7 +485,7 @@ void FMatinee::NewGroupPopupTextCommitted(
 	// if there's no group actor
 	if ( !GroupActor )
 	{
-		if( InActionId == FMatineeCommands::EGroupAction::NewCameraGroup )
+		if( InActionId == FMatineeCommands::EGroupAction::Type::NewCameraGroup )
 		{
 			// find the first perspective viewport - if one exists
 			FLevelEditorViewportClient* ViewportClient = NULL;
@@ -581,15 +581,15 @@ void FMatinee::NewGroupPopupTextCommitted(
 			// For Camera or Skeletal Mesh groups, add a Movement track
 			// FIXME: this doesn't work like this anymore
 			// if you'd like to create multiple groups at once
-			if( InActionId == FMatineeCommands::EGroupAction::NewCameraGroup || 
-				InActionId == FMatineeCommands::EGroupAction::NewSkeletalMeshGroup )
+			if( InActionId == FMatineeCommands::EGroupAction::Type::NewCameraGroup ||
+				InActionId == FMatineeCommands::EGroupAction::Type::NewSkeletalMeshGroup )
 			{
 				int32 NewTrackIndex = INDEX_NONE;
 				AddTrackToGroup( NewGroup, UInterpTrackMove::StaticClass(), NULL, false, NewTrackIndex );
 			}
 
 			// For Camera groups, add a Float Property track for FOV
-			if (InActionId == FMatineeCommands::EGroupAction::NewCameraGroup)
+			if (InActionId == FMatineeCommands::EGroupAction::Type::NewCameraGroup)
 			{
 				// Set the property name for the new track.  This is a global that will be used when setting everything up.
 				SetTrackAddPropName( FName( TEXT( "FOVAngle" ) ) );
@@ -599,7 +599,7 @@ void FMatinee::NewGroupPopupTextCommitted(
 			}
 
 			// For Lighting groups, add a Movement, Brightness, Light Color, and Radius Property track
-			if ( InActionId == FMatineeCommands::EGroupAction::NewLightingGroup )
+			if ( InActionId == FMatineeCommands::EGroupAction::Type::NewLightingGroup )
 			{
 				UInterpTrack* NewMovTrack = NewObject<UInterpTrackMove>(NewGroup, NAME_None, RF_Transactional);
 				const int32 TrackIndex = NewGroup->InterpTracks.Add(NewMovTrack);
@@ -629,14 +629,14 @@ void FMatinee::NewGroupPopupTextCommitted(
 			}
 
 			// For Skeletal Mesh groups, add an Anim track
-			if( InActionId == FMatineeCommands::EGroupAction::NewSkeletalMeshGroup)
+			if( InActionId == FMatineeCommands::EGroupAction::Type::NewSkeletalMeshGroup)
 			{
 				int32 NewTrackIndex = INDEX_NONE;
 				AddTrackToGroup( NewGroup, UInterpTrackAnimControl::StaticClass(), NULL, false, NewTrackIndex );
 			}
 
 			// For Particle groups, add a Toggle track
-			if( InActionId == FMatineeCommands::EGroupAction::NewParticleGroup )
+			if( InActionId == FMatineeCommands::EGroupAction::Type::NewParticleGroup )
 			{
 				int32 NewTrackIndex = INDEX_NONE;
 				AddTrackToGroup( NewGroup, UInterpTrackToggle::StaticClass(), NULL, false, NewTrackIndex );
@@ -1188,7 +1188,7 @@ void FMatinee::OnCreateCameraActorAtCurrentCameraLocation()
 {	
 	// no actor to add
 	TArray<AActor *> OtherActorsToAddToGroup;
-	NewGroupPopup( FMatineeCommands::EGroupAction::NewCameraGroup, NULL, OtherActorsToAddToGroup );
+	NewGroupPopup( FMatineeCommands::EGroupAction::Type::NewCameraGroup, NULL, OtherActorsToAddToGroup );
 }
 
 /** Called when the "Launch Custom Preview Viewport" is pressed */
@@ -2097,7 +2097,7 @@ void FMatinee::OnContextGroupChangeGroupFolder( FMatineeCommands::EGroupAction::
 	check( HasAGroupSelected() );
 
 	// Figure out if we're moving the active group to a new group, or if we simply want to unparent it
-	const bool bIsParenting =  ( InActionId != FMatineeCommands::EGroupAction::RemoveFromGroupFolder ); 
+	const bool bIsParenting =  ( InActionId != FMatineeCommands::EGroupAction::Type::RemoveFromGroupFolder );
 
 	// Figure out which direction we're moving things: A group to the selected folder?  Or, the selected group
 	// to a folder?
@@ -2107,7 +2107,7 @@ void FMatinee::OnContextGroupChangeGroupFolder( FMatineeCommands::EGroupAction::
 	if( bIsParenting )
 	{
 		bIsMovingSelectedGroupToFolder =
-			(InActionId == FMatineeCommands::EGroupAction::MoveActiveGroupToFolder);
+			(InActionId == FMatineeCommands::EGroupAction::Type::MoveActiveGroupToFolder);
 		bIsMovingGroupToSelectedFolder = !bIsMovingSelectedGroupToFolder;
 	}
 
@@ -2233,23 +2233,23 @@ void FMatinee::OnContextKeyInterpMode( FMatineeCommands::EKeyAction::Type InActi
 		FInterpEdSelKey& SelKey = Opt->SelectedKeys[i];
 		UInterpTrack* Track = SelKey.Track;
 
-		if (InActionId == FMatineeCommands::EKeyAction::KeyModeLinear) 
+		if (InActionId == FMatineeCommands::EKeyAction::Type::KeyModeLinear)
 		{
 			Track->SetKeyInterpMode( SelKey.KeyIndex, CIM_Linear );
 		}
-		else if (InActionId == FMatineeCommands::EKeyAction::KeyModeCurveAuto) 
+		else if (InActionId == FMatineeCommands::EKeyAction::Type::KeyModeCurveAuto)
 		{
 			Track->SetKeyInterpMode( SelKey.KeyIndex, CIM_CurveAuto );
 		}
-		else if (InActionId == FMatineeCommands::EKeyAction::KeyModeCurveAutoClamped) 
+		else if (InActionId == FMatineeCommands::EKeyAction::Type::KeyModeCurveAutoClamped)
 		{
 			Track->SetKeyInterpMode( SelKey.KeyIndex, CIM_CurveAutoClamped );
 		}
-		else if(InActionId == FMatineeCommands::EKeyAction::KeyModeCurveBreak) 
+		else if(InActionId == FMatineeCommands::EKeyAction::Type::KeyModeCurveBreak)
 		{
 			Track->SetKeyInterpMode( SelKey.KeyIndex, CIM_CurveBreak );
 		}
-		else if(InActionId == FMatineeCommands::EKeyAction::KeyModeConstant) 
+		else if(InActionId == FMatineeCommands::EKeyAction::Type::KeyModeConstant)
 		{
 			Track->SetKeyInterpMode( SelKey.KeyIndex, CIM_Constant );
 		}
@@ -3713,15 +3713,15 @@ void FMatinee::OnKeyContext_SetCondition( FMatineeCommands::EKeyAction::Type InC
 
 			switch( InCondition )
 			{
-			case FMatineeCommands::EKeyAction::ConditionAlways: 
+			case FMatineeCommands::EKeyAction::Type::ConditionAlways:
 					VisibilityKey.ActiveCondition = EVTC_Always;
 					break;
 
-			case FMatineeCommands::EKeyAction::ConditionGoreEnabled: 
+			case FMatineeCommands::EKeyAction::Type::ConditionGoreEnabled:
 					VisibilityKey.ActiveCondition = EVTC_GoreEnabled;
 					break;
 
-			case FMatineeCommands::EKeyAction::ConditionGoreDisabled: 
+			case FMatineeCommands::EKeyAction::Type::ConditionGoreDisabled:
 					VisibilityKey.ActiveCondition = EVTC_GoreDisabled;
 					break;
 			}
@@ -3745,21 +3745,21 @@ bool FMatinee::KeyContext_IsSetConditionToggled( FMatineeCommands::EKeyAction::T
 
 			switch( InCondition )
 			{
-			case FMatineeCommands::EKeyAction::ConditionAlways: 
+			case FMatineeCommands::EKeyAction::Type::ConditionAlways:
 				if (VisibilityKey.ActiveCondition != EVTC_Always)
 				{
 					return false;
 				}
 				break;
 
-			case FMatineeCommands::EKeyAction::ConditionGoreEnabled: 
+			case FMatineeCommands::EKeyAction::Type::ConditionGoreEnabled:
 				if (VisibilityKey.ActiveCondition != EVTC_GoreEnabled)
 				{
 					return false;
 				}
 				break;
 
-			case FMatineeCommands::EKeyAction::ConditionGoreDisabled: 
+			case FMatineeCommands::EKeyAction::Type::ConditionGoreDisabled:
 				if (VisibilityKey.ActiveCondition != EVTC_GoreDisabled)
 				{
 					return false;
@@ -5643,7 +5643,7 @@ TSharedPtr<SWidget> FMatinee::CreateGroupMenu()
 					FText::FromString( CurrentParent.Group->GroupName.ToString() ),
 					FText::GetEmpty(),
 					FSlateIcon(),
-					FUIAction(FExecuteAction::CreateSP( InMatinee, &FMatinee::OnContextGroupChangeGroupFolder, FMatineeCommands::EGroupAction::MoveActiveGroupToFolder, CurrentParent.GroupIndex))
+					FUIAction(FExecuteAction::CreateSP( InMatinee, &FMatinee::OnContextGroupChangeGroupFolder, FMatineeCommands::EGroupAction::Type::MoveActiveGroupToFolder, CurrentParent.GroupIndex))
 					);
 			}
 		}
@@ -5679,7 +5679,7 @@ TSharedPtr<SWidget> FMatinee::CreateGroupMenu()
 						FText::GetEmpty(),
 						FSlateIcon(),
 						FUIAction(
-							FExecuteAction::CreateSP(InMatinee, &FMatinee::OnContextGroupChangeGroupFolder, FMatineeCommands::EGroupAction::MoveGroupToActiveFolder, CurrentGroupInfo.GroupIndex)
+							FExecuteAction::CreateSP(InMatinee, &FMatinee::OnContextGroupChangeGroupFolder, FMatineeCommands::EGroupAction::Type::MoveGroupToActiveFolder, CurrentGroupInfo.GroupIndex)
 							)
 						);
 
diff --git a/Engine/Source/Editor/VREditor/Public/VREditorInteractor.h b/Engine/Source/Editor/VREditor/Public/VREditorInteractor.h
--- a/Engine/Source/Editor/VREditor/Public/VREditorInteractor.h
+++ b/Engine/Source/Editor/VREditor/Public/VREditorInteractor.h
@@ -25,7 +25,7 @@ enum class EControllerType : uint8
 
 /** Directions the trackpad can be swiped to */
 UENUM( BlueprintType )
-enum ETouchSwipeDirection
+enum class ETouchSwipeDirection : uint8
 {
 	None = 0,
 	Left = 1,
@@ -461,4 +461,4 @@ protected:
 	/** The mode that owns this interactor */
 	UPROPERTY()
 	class UVREditorMode* VRMode;
-};
\ No newline at end of file
+};
diff --git a/Engine/Source/Runtime/ApplicationCore/Private/IOS/IOSInputInterface.cpp b/Engine/Source/Runtime/ApplicationCore/Private/IOS/IOSInputInterface.cpp
--- a/Engine/Source/Runtime/ApplicationCore/Private/IOS/IOSInputInterface.cpp
+++ b/Engine/Source/Runtime/ApplicationCore/Private/IOS/IOSInputInterface.cpp
@@ -126,7 +126,7 @@ void FIOSInputInterface::HandleConnection(GCController* Controller)
 	static_assert(GCControllerPlayerIndex1 == 0 && GCControllerPlayerIndex4 == 3, "Apple changed the player index enums");
 
 	// is this guy a gamepad (i.e., not the Remote)
-	bool bIsGamepadType = (Controller.gamepad != nil);
+	bool bIsGamepadType = (Controller.extendedGamepad != nil);
 	// if we want to use the Remote as a separate player, then we treat it as a Gamepad for player assignment
 	bool bIsTreatedAsGamepad = bIsGamepadType || bTreatRemoteAsSeparateController;
 
@@ -344,7 +344,6 @@ void FIOSInputInterface::SendControllerEvents()
 		
 	for (GCController* Cont in [GCController controllers])
  	{
-		GCGamepad* Gamepad = Cont.gamepad;
 		GCExtendedGamepad* ExtendedGamepad = Cont.extendedGamepad;
 #if PLATFORM_TVOS
 		GCMicroGamepad* MicroGamepad = Cont.microGamepad;
@@ -352,7 +351,7 @@ void FIOSInputInterface::SendControllerEvents()
 		GCMotion* Motion = Cont.motion;
 
 		// skip over gamepads if we don't allow controllers
-		if (Gamepad != nil && !bAllowControllers)
+		if (ExtendedGamepad != nil && !bAllowControllers)
 		{
 			continue;
 		}
@@ -456,29 +455,6 @@ if ((Controller.Previous##Gamepad != nil && Gamepad.GCAxis.value != Controller.P
 			Controller.PreviousExtendedGamepad = [ExtendedGamepad saveSnapshot];
 			[Controller.PreviousExtendedGamepad retain];
 		}
-		// get basic input (extended is a superset, don't do both)
-		else if (Gamepad != nil)
-		{
-			HANDLE_BUTTON(Gamepad, buttonA,			FGamepadKeyNames::FaceButtonBottom);
-			HANDLE_BUTTON(Gamepad, buttonA,			FGamepadKeyNames::FaceButtonBottom);
-			HANDLE_BUTTON(Gamepad, buttonB,			FGamepadKeyNames::FaceButtonRight);
-			HANDLE_BUTTON(Gamepad, buttonX,			FGamepadKeyNames::FaceButtonLeft);
-			HANDLE_BUTTON(Gamepad, buttonY,			FGamepadKeyNames::FaceButtonTop);
-			HANDLE_BUTTON(Gamepad, leftShoulder,	FGamepadKeyNames::LeftShoulder);
-			HANDLE_BUTTON(Gamepad, rightShoulder,	FGamepadKeyNames::RightShoulder);
-			HANDLE_BUTTON(Gamepad, dpad.up,			FGamepadKeyNames::DPadUp);
-			HANDLE_BUTTON(Gamepad, dpad.down,		FGamepadKeyNames::DPadDown);
-			HANDLE_BUTTON(Gamepad, dpad.right,		FGamepadKeyNames::DPadRight);
-			HANDLE_BUTTON(Gamepad, dpad.left,		FGamepadKeyNames::DPadLeft);
-			
-
-			HANDLE_ANALOG(ExtendedGamepad, dpad.xAxis,	FGamepadKeyNames::LeftAnalogX);
-			HANDLE_ANALOG(ExtendedGamepad, dpad.yAxis,	FGamepadKeyNames::LeftAnalogY);
-			
-			[Controller.PreviousGamepad release];
-			Controller.PreviousGamepad = [Gamepad saveSnapshot];
-			[Controller.PreviousGamepad retain];
-        }
 #if PLATFORM_TVOS
         // get micro input (shouldn't have the other two)
         else if (MicroGamepad != nil)
diff --git a/Engine/Source/Runtime/Engine/Classes/Engine/Engine.h b/Engine/Source/Runtime/Engine/Classes/Engine/Engine.h
--- a/Engine/Source/Runtime/Engine/Classes/Engine/Engine.h
+++ b/Engine/Source/Runtime/Engine/Classes/Engine/Engine.h
@@ -90,25 +90,25 @@ enum EFullyLoadPackageType
  * Enumerates transition types.
  */
 UENUM()
-enum ETransitionType
+enum class ETransitionType : uint8
 {
-	TT_None,
-	TT_Paused,
-	TT_Loading,
-	TT_Saving,
-	TT_Connecting,
-	TT_Precaching,
-	TT_WaitingToConnect,
-	TT_MAX,
+	None,
+	Paused,
+	Loading,
+	Saving,
+	Connecting,
+	Precaching,
+	WaitingToConnect,
+	MAX
 };
 
 
 UENUM()
-enum EConsoleType
+enum class EConsoleType : uint8
 {
-	CONSOLE_Any,
-	CONSOLE_Mobile,
-	CONSOLE_MAX,
+	Any,
+	Mobile,
+	MAX
 };
 
 /** Status of dynamic resolution that depends on project setting cvar, game user settings, and pause */
@@ -1458,7 +1458,7 @@ public:
 
 	/** The current transition type. */
 	UPROPERTY()
-	TEnumAsByte<enum ETransitionType> TransitionType;
+	ETransitionType TransitionType;
 
 	/** The current transition description text. */
 	UPROPERTY()
@@ -2346,7 +2346,7 @@ public:
 	 *
 	 * @return true if we're on a console, false if we're running on a PC
 	 */
-	bool IsConsoleBuild(EConsoleType ConsoleType = CONSOLE_Any) const;
+	bool IsConsoleBuild(EConsoleType ConsoleType = EConsoleType::Any) const;
 
 	/** Add a FString to the On-screen debug message system. bNewerOnTop only works with Key == INDEX_NONE */
 	void AddOnScreenDebugMessage(uint64 Key,float TimeToDisplay,FColor DisplayColor,const FString& DebugMessage, bool bNewerOnTop = true, const FVector2D& TextScale = FVector2D::UnitVector);
diff --git a/Engine/Source/Runtime/Engine/Classes/Engine/StaticMesh.h b/Engine/Source/Runtime/Engine/Classes/Engine/StaticMesh.h
--- a/Engine/Source/Runtime/Engine/Classes/Engine/StaticMesh.h
+++ b/Engine/Source/Runtime/Engine/Classes/Engine/StaticMesh.h
@@ -450,8 +450,8 @@ enum EImportStaticMeshVersion
 	BeforeImportStaticMeshVersionWasAdded,
 	// Remove the material re-order workflow
 	RemoveStaticMeshSkinxxWorkflow,
-	VersionPlusOne,
-	LastVersion = VersionPlusOne - 1
+	StaticMeshVersionPlusOne,
+	LastVersion = StaticMeshVersionPlusOne - 1
 };
 
 USTRUCT()
diff --git a/Engine/Source/Runtime/Engine/Classes/Materials/MaterialFunctionInterface.h b/Engine/Source/Runtime/Engine/Classes/Materials/MaterialFunctionInterface.h
--- a/Engine/Source/Runtime/Engine/Classes/Materials/MaterialFunctionInterface.h
+++ b/Engine/Source/Runtime/Engine/Classes/Materials/MaterialFunctionInterface.h
@@ -20,7 +20,7 @@ struct FPropertyChangedEvent;
 
 /** Usage set on a material function determines feature compatibility and validation. */
 UENUM()
-enum EMaterialFunctionUsage
+enum class EMaterialFunctionUsage : uint8
 {
 	Default,
 	MaterialLayer,
@@ -46,7 +46,7 @@ class UMaterialFunctionInterface : public UObject
 protected:
 	/** The intended usage of this function, required for material layers. */
 	UPROPERTY(AssetRegistrySearchable)
-	TEnumAsByte<enum EMaterialFunctionUsage> MaterialFunctionUsage;
+	EMaterialFunctionUsage MaterialFunctionUsage;
 
 public:
 	virtual EMaterialFunctionUsage GetMaterialFunctionUsage()
diff --git a/Engine/Source/Runtime/Engine/Private/GameInstance.cpp b/Engine/Source/Runtime/Engine/Private/GameInstance.cpp
--- a/Engine/Source/Runtime/Engine/Private/GameInstance.cpp
+++ b/Engine/Source/Runtime/Engine/Private/GameInstance.cpp
@@ -321,7 +321,7 @@ FGameInstancePIEResult UGameInstance::StartPlayInEditorGameInstance(ULocalPlayer
 
 		if (EditorEngine->Browse(*WorldContext, FURL(&BaseURL, *URLString, (ETravelType)TRAVEL_Absolute), Error) == EBrowseReturnVal::Pending)
 		{
-			EditorEngine->TransitionType = TT_WaitingToConnect;
+			EditorEngine->TransitionType = ETransitionType::WaitingToConnect;
 		}
 		else
 		{
@@ -1353,4 +1353,4 @@ void UGameInstance::RegisterReferencedObject(UObject* ObjectToReference)
 void UGameInstance::UnregisterReferencedObject(UObject* ObjectToReference)
 {
 	ReferencedObjects.RemoveSingleSwap(ObjectToReference);
-}
\ No newline at end of file
+}
diff --git a/Engine/Source/Runtime/Engine/Private/GameViewportClient.cpp b/Engine/Source/Runtime/Engine/Private/GameViewportClient.cpp
--- a/Engine/Source/Runtime/Engine/Private/GameViewportClient.cpp
+++ b/Engine/Source/Runtime/Engine/Private/GameViewportClient.cpp
@@ -2475,22 +2475,22 @@ void UGameViewportClient::DrawTransition(UCanvas* Canvas)
 	{
 		switch (GetOuterUEngine()->TransitionType)
 		{
-		case TT_Loading:
+		case ETransitionType::Loading:
 			DrawTransitionMessage(Canvas, NSLOCTEXT("GameViewportClient", "LoadingMessage", "LOADING").ToString());
 			break;
-		case TT_Saving:
+		case ETransitionType::Saving:
 			DrawTransitionMessage(Canvas, NSLOCTEXT("GameViewportClient", "SavingMessage", "SAVING").ToString());
 			break;
-		case TT_Connecting:
+		case ETransitionType::Connecting:
 			DrawTransitionMessage(Canvas, NSLOCTEXT("GameViewportClient", "ConnectingMessage", "CONNECTING").ToString());
 			break;
-		case TT_Precaching:
+		case ETransitionType::Precaching:
 			DrawTransitionMessage(Canvas, NSLOCTEXT("GameViewportClient", "PrecachingMessage", "PRECACHING").ToString());
 			break;
-		case TT_Paused:
+		case ETransitionType::Paused:
 			DrawTransitionMessage(Canvas, NSLOCTEXT("GameViewportClient", "PausedMessage", "PAUSED").ToString());
 			break;
-		case TT_WaitingToConnect:
+		case ETransitionType::WaitingToConnect:
 			DrawTransitionMessage(Canvas, TEXT("Waiting to connect...")); // Temp - localization of the FString messages is broke atm. Loc this when its fixed.
 			break;
 		}
diff --git a/Engine/Source/Runtime/Engine/Private/PendingNetGame.cpp b/Engine/Source/Runtime/Engine/Private/PendingNetGame.cpp
--- a/Engine/Source/Runtime/Engine/Private/PendingNetGame.cpp
+++ b/Engine/Source/Runtime/Engine/Private/PendingNetGame.cpp
@@ -150,7 +150,7 @@ void UPendingNetGame::LoadMapCompleted(UEngine* Engine, FWorldContext& Context,
 	else
 	{
 		// Show connecting message, cause precaching to occur.
-		Engine->TransitionType = TT_Connecting;
+		Engine->TransitionType = ETransitionType::Connecting;
 
 		Engine->RedrawViewports(false);
 
diff --git a/Engine/Source/Runtime/Engine/Private/UnrealEngine.cpp b/Engine/Source/Runtime/Engine/Private/UnrealEngine.cpp
--- a/Engine/Source/Runtime/Engine/Private/UnrealEngine.cpp
+++ b/Engine/Source/Runtime/Engine/Private/UnrealEngine.cpp
@@ -8846,13 +8846,13 @@ bool UEngine::IsConsoleBuild(EConsoleType ConsoleType) const
 {
 	switch (ConsoleType)
 	{
-	case CONSOLE_Any:
+	case EConsoleType::Any:
 #if !PLATFORM_DESKTOP
 		return true;
 #else
 		return false;
 #endif
-	case CONSOLE_Mobile:
+	case EConsoleType::Mobile:
 		return false;
 	default:
 		UE_LOG(LogEngine, Warning, TEXT("Unknown ConsoleType passed to IsConsoleBuild()"));
@@ -12106,9 +12106,9 @@ void UEngine::TickWorldTravel(FWorldContext& Context, float DeltaSeconds)
 			}
 		}
 	}
-	else if (TransitionType == TT_WaitingToConnect)
+	else if (TransitionType == ETransitionType::WaitingToConnect)
 	{
-		TransitionType = TT_None;
+		TransitionType = ETransitionType::None;
 	}
 
 	return;
@@ -12214,7 +12214,7 @@ bool UEngine::LoadMap( FWorldContext& WorldContext, FURL URL, class UPendingNetG
 
 		if(!URL.HasOption(TEXT("quiet")) )
 		{
-			TransitionType = TT_Loading;
+			TransitionType = ETransitionType::Loading;
 			TransitionDescription = URL.Map;
 			if (URL.HasOption(TEXT("Game=")))
 			{
@@ -12233,7 +12233,7 @@ bool UEngine::LoadMap( FWorldContext& WorldContext, FURL URL, class UPendingNetG
 				LoadMapRedrawViewports();
 			}
 
-			TransitionType = TT_None;
+			TransitionType = ETransitionType::None;
 		}
 
 		// Clean up networking
@@ -12918,10 +12918,10 @@ void UEngine::LoadPackagesFully(UWorld * InWorld, EFullyLoadPackageType FullyLoa
 void UEngine::UpdateTransitionType(UWorld *CurrentWorld)
 {
 	// Update the transition screen.
-	if(TransitionType == TT_Connecting)
+	if(TransitionType == ETransitionType::Connecting)
 	{
 		// Check to see if all players have finished connecting.
-		TransitionType = TT_None;
+		TransitionType = ETransitionType::None;
 
 		FWorldContext &Context = GetWorldContextFromWorldChecked(CurrentWorld);
 		if (Context.OwningGameInstance != NULL)
@@ -12931,16 +12931,16 @@ void UEngine::UpdateTransitionType(UWorld *CurrentWorld)
 				if(!(*It)->PlayerController)
 				{
 					// This player has not received a PlayerController from the server yet, so leave the connecting screen up.
-					TransitionType = TT_Connecting;
+					TransitionType = ETransitionType::Connecting;
 					break;
 				}
 			}
 		}
 	}
-	else if(TransitionType == TT_None || TransitionType == TT_Paused)
+	else if(TransitionType == ETransitionType::None || TransitionType == ETransitionType::Paused)
 	{
 		// Display a paused screen if the game is paused.
-		TransitionType = (CurrentWorld->GetWorldSettings()->GetPauserPlayerState() != NULL) ? TT_Paused : TT_None;
+		TransitionType = (CurrentWorld->GetWorldSettings()->GetPauserPlayerState() != NULL) ? ETransitionType::Paused : ETransitionType::None;
 	}
 }
 
diff --git a/Engine/Source/Runtime/MovieScene/Private/MovieSceneSequencePlayer.cpp b/Engine/Source/Runtime/MovieScene/Private/MovieSceneSequencePlayer.cpp
--- a/Engine/Source/Runtime/MovieScene/Private/MovieSceneSequencePlayer.cpp
+++ b/Engine/Source/Runtime/MovieScene/Private/MovieSceneSequencePlayer.cpp
@@ -214,7 +214,7 @@ void UMovieSceneSequencePlayer::Pause()
 	{
 		if (bIsEvaluating)
 		{
-			LatentActions.Emplace(FLatentAction::Pause);
+			LatentActions.Emplace(FLatentAction::EType::Pause);
 			return;
 		}
 
@@ -283,7 +283,7 @@ void UMovieSceneSequencePlayer::StopInternal(FFrameTime TimeToResetTo)
 	{
 		if (bIsEvaluating)
 		{
-			LatentActions.Emplace(FLatentAction::Stop, TimeToResetTo);
+			LatentActions.Emplace(FLatentAction::EType::Stop, TimeToResetTo);
 			return;
 		}
 
@@ -891,11 +891,11 @@ void UMovieSceneSequencePlayer::ApplyLatentActions()
 	{
 		switch (LatentAction.Type)
 		{
-		case FLatentAction::Stop:          StopInternal(LatentAction.Position); continue;
-		case FLatentAction::Pause:         Pause();                             continue;
+		case FLatentAction::EType::Stop:   StopInternal(LatentAction.Position); continue;
+		case FLatentAction::EType::Pause:  Pause();                             continue;
 		}
 
-		check(LatentAction.Type == FLatentAction::Update);
+		check(LatentAction.Type == FLatentAction::EType::Update);
 		switch (LatentAction.UpdateMethod)
 		{
 		case EUpdatePositionMethod::Play:  PlayToFrame( LatentAction.Position); continue;
diff --git a/Engine/Source/Runtime/MovieScene/Public/MovieSceneSequencePlayer.h b/Engine/Source/Runtime/MovieScene/Public/MovieSceneSequencePlayer.h
--- a/Engine/Source/Runtime/MovieScene/Public/MovieSceneSequencePlayer.h
+++ b/Engine/Source/Runtime/MovieScene/Public/MovieSceneSequencePlayer.h
@@ -630,14 +630,14 @@ protected:
 
 	struct FLatentAction
 	{
-		enum EType { Stop, Pause, Update };
+		enum class EType : uint8 { Stop, Pause, Update };
 
 		FLatentAction(EType InType, FFrameTime DesiredTime = 0)
 			: Type(InType)
 		{}
 
 		FLatentAction(EUpdatePositionMethod InUpdateMethod, FFrameTime DesiredTime)
-			: Type(Update), UpdateMethod(InUpdateMethod), Position(DesiredTime)
+			: Type(EType::Update), UpdateMethod(InUpdateMethod), Position(DesiredTime)
 		{}
 
 		EType                 Type;
diff --git a/Engine/Source/Runtime/Online/BuildPatchServices/Private/Common/StatsCollector.h b/Engine/Source/Runtime/Online/BuildPatchServices/Private/Common/StatsCollector.h
--- a/Engine/Source/Runtime/Online/BuildPatchServices/Private/Common/StatsCollector.h
+++ b/Engine/Source/Runtime/Online/BuildPatchServices/Private/Common/StatsCollector.h
@@ -5,7 +5,7 @@
 
 namespace BuildPatchServices
 {
-	enum EStatFormat
+	enum class EStatFormat : uint8
 	{
 		// Using the AccumulateTimeBegin and End functions, measured in cycles
 		Timer,
diff --git a/Engine/Source/Runtime/OpenGLDrv/Private/IOS/IOSOpenGL.h b/Engine/Source/Runtime/OpenGLDrv/Private/IOS/IOSOpenGL.h
--- a/Engine/Source/Runtime/OpenGLDrv/Private/IOS/IOSOpenGL.h
+++ b/Engine/Source/Runtime/OpenGLDrv/Private/IOS/IOSOpenGL.h
@@ -102,10 +102,10 @@ struct FIOSOpenGL : public FOpenGLES2
 
 	static FORCEINLINE void* MapBufferRange(GLenum Type, uint32 InOffset, uint32 InSize, EResourceLockMode LockMode)
 	{
-		checkf(LockMode == RLM_WriteOnly || LockMode == RLM_WriteOnlyUnsynchronized, TEXT("OpenGL ES 2.0 only supports write-only buffer locks"));
+		checkf(LockMode == EResourceLockMode::RLM_WriteOnly || LockMode == EResourceLockMode::RLM_WriteOnlyUnsynchronized, TEXT("OpenGL ES 2.0 only supports write-only buffer locks"));
 		check(Type == GL_ARRAY_BUFFER || Type == GL_ELEMENT_ARRAY_BUFFER);
 		GLuint Flags = GL_MAP_WRITE_BIT_EXT | GL_MAP_FLUSH_EXPLICIT_BIT_EXT;
-		if (LockMode == RLM_WriteOnlyUnsynchronized)
+		if (LockMode == EResourceLockMode::RLM_WriteOnlyUnsynchronized)
 		{
 			Flags |= GL_MAP_UNSYNCHRONIZED_BIT_EXT;
 		}
diff --git a/Engine/Source/Runtime/OpenGLDrv/Private/OpenGLUniformBuffer.cpp b/Engine/Source/Runtime/OpenGLDrv/Private/OpenGLUniformBuffer.cpp
--- a/Engine/Source/Runtime/OpenGLDrv/Private/OpenGLUniformBuffer.cpp
+++ b/Engine/Source/Runtime/OpenGLDrv/Private/OpenGLUniformBuffer.cpp
@@ -346,7 +346,7 @@ static void SuballocateUBO( uint32 Size, GLuint& Resource, uint32& Offset, uint8
 	if (FOpenGL::SupportsBufferStorage() && OpenGLConsoleVariables::bUBODirectWrite)
 	{
 		FOpenGL::BufferStorage( GL_UNIFORM_BUFFER, GetUBOPoolSize(), NULL, GL_MAP_WRITE_BIT | GL_MAP_PERSISTENT_BIT | GL_MAP_COHERENT_BIT );
-		Pool.Pointer = (uint8*)FOpenGL::MapBufferRange(GL_UNIFORM_BUFFER, 0, GetUBOPoolSize(), FOpenGL::RLM_WriteOnlyPersistent);
+		Pool.Pointer = (uint8*)FOpenGL::MapBufferRange(GL_UNIFORM_BUFFER, 0, GetUBOPoolSize(), FOpenGL::EResourceLockMode::RLM_WriteOnlyPersistent);
 	}
 	else
 	{
@@ -803,4 +803,4 @@ void FOpenGLDynamicRHI::RHIUpdateUniformBuffer(FRHIUniformBuffer* UniformBufferR
 		});
 		RHICmdList.RHIThreadFence(true);
 	}
-}
\ No newline at end of file
+}
diff --git a/Engine/Source/Runtime/OpenGLDrv/Private/OpenGLVertexBuffer.cpp b/Engine/Source/Runtime/OpenGLDrv/Private/OpenGLVertexBuffer.cpp
--- a/Engine/Source/Runtime/OpenGLDrv/Private/OpenGLVertexBuffer.cpp
+++ b/Engine/Source/Runtime/OpenGLDrv/Private/OpenGLVertexBuffer.cpp
@@ -61,7 +61,7 @@ void* GetAllocation( void* Target, uint32 Size, uint32 Offset, uint32 Alignment
 			FOpenGL::GenBuffers(1, &PoolVB);
 			glBindBuffer(GL_COPY_READ_BUFFER, PoolVB);
 			FOpenGL::BufferStorage(GL_COPY_READ_BUFFER, PerFrameMax * 4, NULL, GL_MAP_WRITE_BIT | GL_MAP_PERSISTENT_BIT | GL_MAP_COHERENT_BIT);
-			PoolPointer = (uint8*)FOpenGL::MapBufferRange(GL_COPY_READ_BUFFER, 0, PerFrameMax * 4, FOpenGL::RLM_WriteOnlyPersistent);
+			PoolPointer = (uint8*)FOpenGL::MapBufferRange(GL_COPY_READ_BUFFER, 0, PerFrameMax * 4, FOpenGL::EResourceLockMode::RLM_WriteOnlyPersistent);
 
 			FreeSpace = PerFrameMax * 4;
 
@@ -209,7 +209,7 @@ void* FOpenGLDynamicRHI::RHILockVertexBuffer(FRHIVertexBuffer* VertexBufferRHI,u
 	}
 	else
 	{
-		if (VertexBuffer->IsDynamic() && LockMode == RLM_WriteOnly)
+		if (VertexBuffer->IsDynamic() && LockMode == EResourceLockMode::RLM_WriteOnly)
 		{
 			void *Staging = GetAllocation(VertexBuffer, Size, Offset);
 			if (Staging)
@@ -217,7 +217,7 @@ void* FOpenGLDynamicRHI::RHILockVertexBuffer(FRHIVertexBuffer* VertexBufferRHI,u
 				return Staging;
 			}
 		}
-		return (void*)VertexBuffer->Lock(Offset, Size, LockMode == RLM_ReadOnly, VertexBuffer->IsDynamic());
+		return (void*)VertexBuffer->Lock(Offset, Size, LockMode == EResourceLockMode::RLM_ReadOnly, VertexBuffer->IsDynamic());
 	}
 	RHITHREAD_GLCOMMAND_EPILOGUE_RETURN(void*);
 }
diff --git a/Engine/Source/Runtime/OpenGLDrv/Public/OpenGL.h b/Engine/Source/Runtime/OpenGLDrv/Public/OpenGL.h
--- a/Engine/Source/Runtime/OpenGLDrv/Public/OpenGL.h
+++ b/Engine/Source/Runtime/OpenGLDrv/Public/OpenGL.h
@@ -51,7 +51,7 @@ struct FPlatformOpenGLContext;
 class FOpenGLBase
 {
 public:
-	enum EResourceLockMode
+	enum class EResourceLockMode : uint8
 	{
 		RLM_ReadWrite,
 		RLM_ReadOnly,
diff --git a/Engine/Source/Runtime/OpenGLDrv/Public/OpenGL3.h b/Engine/Source/Runtime/OpenGLDrv/Public/OpenGL3.h
--- a/Engine/Source/Runtime/OpenGLDrv/Public/OpenGL3.h
+++ b/Engine/Source/Runtime/OpenGLDrv/Public/OpenGL3.h
@@ -142,23 +142,23 @@ struct FOpenGL3 : public FOpenGLBase
 		GLenum Access;
 		switch ( LockMode )
 		{
-			case RLM_ReadOnly:
+			case EResourceLockMode::RLM_ReadOnly:
 				Access = GL_MAP_READ_BIT;
 				break;
-			case RLM_WriteOnly:
+			case EResourceLockMode::RLM_WriteOnly:
 				Access = (GL_MAP_INVALIDATE_BUFFER_BIT | GL_MAP_WRITE_BIT);
 #if 1
 				// Temp workaround for synchrnoization when a UBO is discarded while being referenced
 				Access |= GL_MAP_UNSYNCHRONIZED_BIT;
 #endif
 				break;
-			case RLM_WriteOnlyUnsynchronized:
+			case EResourceLockMode::RLM_WriteOnlyUnsynchronized:
 				Access = (GL_MAP_WRITE_BIT | GL_MAP_UNSYNCHRONIZED_BIT);
 				break;
-			case RLM_WriteOnlyPersistent:
+			case EResourceLockMode::RLM_WriteOnlyPersistent:
 				Access = (GL_MAP_WRITE_BIT | GL_MAP_PERSISTENT_BIT | GL_MAP_COHERENT_BIT);
 				break;
-			case RLM_ReadWrite:
+			case EResourceLockMode::RLM_ReadWrite:
 			default:
 				Access = (GL_MAP_READ_BIT | GL_MAP_WRITE_BIT);
 		}
diff --git a/Engine/Source/Runtime/OpenGLDrv/Public/OpenGLES2.h b/Engine/Source/Runtime/OpenGLDrv/Public/OpenGLES2.h
--- a/Engine/Source/Runtime/OpenGLDrv/Public/OpenGLES2.h
+++ b/Engine/Source/Runtime/OpenGLDrv/Public/OpenGLES2.h
@@ -245,9 +245,9 @@ struct FOpenGLES2 : public FOpenGLBase
 #if OPENGL_ES2_BRING_UP
 		// Non-written areas retain prior values.
 		// Lack of unsynchronized in glMapBufferOES() is a perf bug which needs to be fixed later.
-		checkf(LockMode == RLM_WriteOnly || LockMode == RLM_WriteOnlyUnsynchronized, TEXT("OpenGL ES 2.0 only supports write-only buffer locks"));
+		checkf(LockMode == EResourceLockMode::RLM_WriteOnly || LockMode == EResourceLockMode::RLM_WriteOnlyUnsynchronized, TEXT("OpenGL ES 2.0 only supports write-only buffer locks"));
 #else
-		checkf(LockMode == RLM_WriteOnly, TEXT("OpenGL ES 2.0 only supports write-only buffer locks"));
+		checkf(LockMode == EResourceLockMode::RLM_WriteOnly, TEXT("OpenGL ES 2.0 only supports write-only buffer locks"));
 #endif
 		check(Type == GL_ARRAY_BUFFER || Type == GL_ELEMENT_ARRAY_BUFFER);
 
diff --git a/Engine/Source/Runtime/OpenGLDrv/Public/OpenGLResources.h b/Engine/Source/Runtime/OpenGLDrv/Public/OpenGLResources.h
--- a/Engine/Source/Runtime/OpenGLDrv/Public/OpenGLResources.h
+++ b/Engine/Source/Runtime/OpenGLDrv/Public/OpenGLResources.h
@@ -618,7 +618,7 @@ public:
 
 		if ( bUseMapBuffer)
 		{
-			FOpenGL::EResourceLockMode LockMode = bReadOnly ? FOpenGL::RLM_ReadOnly : FOpenGL::RLM_WriteOnly;
+			FOpenGL::EResourceLockMode LockMode = bReadOnly ? FOpenGL::EResourceLockMode::RLM_ReadOnly : FOpenGL::EResourceLockMode::RLM_WriteOnly;
 			Data = static_cast<uint8*>( FOpenGL::MapBufferRange( Type, InOffset, InSize, LockMode ) );
 //			checkf(Data != NULL, TEXT("FOpenGL::MapBufferRange Failed, glError %d (0x%x)"), glGetError(), glGetError());
 
@@ -685,7 +685,7 @@ public:
 
 		if ( bUseMapBuffer)
 		{
-			FOpenGL::EResourceLockMode LockMode = bDiscard ? FOpenGL::RLM_WriteOnly : FOpenGL::RLM_WriteOnlyUnsynchronized;
+			FOpenGL::EResourceLockMode LockMode = bDiscard ? FOpenGL::EResourceLockMode::RLM_WriteOnly : FOpenGL::EResourceLockMode::RLM_WriteOnlyUnsynchronized;
 			Data = static_cast<uint8*>( FOpenGL::MapBufferRange( Type, InOffset, InSize, LockMode ) );
 			LockOffset = InOffset;
 			LockSize = InSize;
