Added a deprecation message to the Matinee UI warning users that this will be the last version to support Matinee. Added analytic events for when people convert Matinee to Sequencer and for when people dismiss the warning dialog.

diff --git a/Engine/Plugins/MovieScene/MatineeToLevelSequence/Source/MatineeToLevelSequence/MatineeToLevelSequence.Build.cs b/Engine/Plugins/MovieScene/MatineeToLevelSequence/Source/MatineeToLevelSequence/MatineeToLevelSequence.Build.cs
--- a/Engine/Plugins/MovieScene/MatineeToLevelSequence/Source/MatineeToLevelSequence/MatineeToLevelSequence.Build.cs
+++ b/Engine/Plugins/MovieScene/MatineeToLevelSequence/Source/MatineeToLevelSequence/MatineeToLevelSequence.Build.cs
@@ -26,7 +26,8 @@ public class MatineeToLevelSequence : ModuleRules
                 "Slate",
                 "SlateCore",
                 "UnrealEd",
-				"TimeManagement"
+				"TimeManagement",
+				"Analytics"
             }
         );
 
diff --git a/Engine/Plugins/MovieScene/MatineeToLevelSequence/Source/MatineeToLevelSequence/Private/MatineeToLevelSequenceModule.cpp b/Engine/Plugins/MovieScene/MatineeToLevelSequence/Source/MatineeToLevelSequence/Private/MatineeToLevelSequenceModule.cpp
--- a/Engine/Plugins/MovieScene/MatineeToLevelSequence/Source/MatineeToLevelSequence/Private/MatineeToLevelSequenceModule.cpp
+++ b/Engine/Plugins/MovieScene/MatineeToLevelSequence/Source/MatineeToLevelSequence/Private/MatineeToLevelSequenceModule.cpp
@@ -78,6 +78,8 @@
 #include "MovieSceneTimeHelpers.h"
 #include "Components/SkeletalMeshComponent.h"
 #include "Engine/StaticMeshActor.h"
+#include "EngineAnalytics.h"
+#include "Interfaces/IAnalyticsProvider.h"
 
 #define LOCTEXT_NAMESPACE "MatineeToLevelSequence"
 
@@ -212,7 +214,7 @@ protected:
 
 		// Pop open a dialog asking whether the user to convert and launcher sequencer or no
 		FSuppressableWarningDialog::FSetupInfo Info( 
-			LOCTEXT("MatineeToLevelSequencePrompt", "Matinee is now a legacy tool. Would you like to continue opening Matinee or convert your Matinee to a Level Sequence Asset?"), 
+			LOCTEXT("MatineeToLevelSequencePrompt", "As of 4.23, Matinee is no longer supported by UE4 and will be removed from the engine in the near future. Once removed, you will no longer be able to run a Matinee or open Matinee.\n\nWould you like to continue opening Matinee or convert your Matinee to a Level Sequence Asset?"), 
 			LOCTEXT("MatineeToLevelSequenceTitle", "Convert Matinee to Level Sequence Asset"), 
 			TEXT("MatineeToLevelSequence") );
 		Info.ConfirmText = LOCTEXT("MatineeToLevelSequence_ConfirmText", "Open Matinee");
@@ -237,6 +239,12 @@ protected:
 	/** Callback for converting a matinee to a level sequence asset. */
 	void OnConvertMatineeToLevelSequence(TArray<TWeakObjectPtr<AActor> > ActorsToConvert)
 	{
+		// Keep track of how many people actually used the tool to convert assets over.
+		if (FEngineAnalytics::IsAvailable())
+		{
+			FEngineAnalytics::GetProvider().RecordEvent(TEXT("Editor.Matinee.ConversionTool.MatineeActorConverted"));
+		}
+
 		TArray<TWeakObjectPtr<ALevelSequenceActor> > NewActors;
 
 		int32 NumWarnings = 0;
diff --git a/Engine/Source/Editor/Matinee/Matinee.Build.cs b/Engine/Source/Editor/Matinee/Matinee.Build.cs
--- a/Engine/Source/Editor/Matinee/Matinee.Build.cs
+++ b/Engine/Source/Editor/Matinee/Matinee.Build.cs
@@ -21,6 +21,7 @@ public class Matinee : ModuleRules
                 "UnrealEd",	//compatibility for FBX exporter
                 "MainFrame",
                 "WorkspaceMenuStructure",
+				"Analytics"
             }
             );
 
diff --git a/Engine/Source/Editor/Matinee/Private/MatineeTrackView.cpp b/Engine/Source/Editor/Matinee/Private/MatineeTrackView.cpp
--- a/Engine/Source/Editor/Matinee/Private/MatineeTrackView.cpp
+++ b/Engine/Source/Editor/Matinee/Private/MatineeTrackView.cpp
@@ -12,6 +12,102 @@
 
 #include "Slate/SceneViewport.h"
 #include "Widgets/Docking/SDockTab.h"
+#include "Widgets/Text/SRichTextBlock.h"
+#include "Widgets/Input/SButton.h"
+#include "EngineAnalytics.h"
+#include "Interfaces/IAnalyticsProvider.h"
+#include "IDocumentation.h"
+#include "Framework/Text/SlateHyperlinkRun.h"
+
+static void OnDocLinkClicked(const FSlateHyperlinkRun::FMetadata& Metadata)
+{
+	const FString* Url = Metadata.Find(TEXT("href"));
+	if (Url)
+	{
+		if (FEngineAnalytics::IsAvailable())
+		{
+			TArray<FAnalyticsEventAttribute> EventAttributes;
+			EventAttributes.Add(FAnalyticsEventAttribute(TEXT("DocLink"), *Url));
+
+			FEngineAnalytics::GetProvider().RecordEvent(TEXT("Editor.Matinee.DeprecationWarning.DocLinkClicked"), EventAttributes);
+		}
+
+		IDocumentation::Get()->Open(*Url, FDocumentationSourceInfo(TEXT("editor")));
+	}
+}
+
+class SMatineeDeprecationMessage : public SCompoundWidget
+{
+public:
+
+	SLATE_BEGIN_ARGS(SMatineeDeprecationMessage)
+	{}
+	SLATE_END_ARGS()
+
+	void Construct(const FArguments& InArgs)
+	{
+		ChildSlot
+		[
+			SNew(SBorder)
+			.BorderBackgroundColor(FLinearColor(0.4f, 0.f, 0.f, 1.f))
+			.BorderImage(FEditorStyle::GetBrush(TEXT("WhiteBrush")))
+			.Visibility(this, &SMatineeDeprecationMessage::GetWarningMessageVisibility)
+			[
+				SNew(SHorizontalBox)
+				+SHorizontalBox::Slot()
+				.HAlign(HAlign_Fill)
+				.VAlign(VAlign_Center)
+				.Padding(4, 0, 0, 0)
+				[
+					SNew(SRichTextBlock)
+					.Text(NSLOCTEXT("Matinee", "MatineeLastVersionSupported", "As of 4.23, Matinee is no longer supported by UE4 and will be removed from the engine in the near future. Once removed, you will <NormalText.Important>no longer be able to run a Matinee or open Matinee Editor</>.\n"
+									"Please use the <a id=\"udn\" href=\"/Engine/Sequencer/HowTo/MatineeConversionTool\" style=\"Hyperlink\">Matinee to Sequencer Conversion Tool</> to convert any files to Sequencer as soon as possible."))
+					.AutoWrapText(true)
+					.DecoratorStyleSet(&FEditorStyle::Get())
+					+ SRichTextBlock::HyperlinkDecorator(TEXT("udn"), FSlateHyperlinkRun::FOnClick::CreateStatic(&OnDocLinkClicked))
+
+				]
+				+SHorizontalBox::Slot()
+				.HAlign(HAlign_Right)
+				.VAlign(VAlign_Center)
+				.AutoWidth()
+				[
+					SNew(SButton)
+					.Text(NSLOCTEXT("Matinee", "DismissMatineeSupportWarning", "Dismiss"))
+					.OnClicked(this, &SMatineeDeprecationMessage::DismissWarningForever)
+				]
+			]
+		];
+	}
+
+	FReply DismissWarningForever() 
+	{
+		if (GConfig)
+		{
+			GConfig->SetBool(TEXT("Matinee"), TEXT("HasDismissedDeprecationWarning"), true, GEditorIni);
+
+			if (FEngineAnalytics::IsAvailable())
+			{
+				FEngineAnalytics::GetProvider().RecordEvent(TEXT("Editor.Matinee.DeprecationWarning.Dimissed"));
+			}
+		}
+
+		return FReply::Handled();
+	}
+
+	EVisibility GetWarningMessageVisibility() const
+	{
+		if (GConfig)
+		{
+			bool bDismissed = false;
+			GConfig->GetBool(TEXT("Matinee"), TEXT("HasDismissedDeprecationWarning"), bDismissed, GEditorIni);
+
+			return bDismissed ? EVisibility::Collapsed : EVisibility::Visible;
+		}
+
+		return EVisibility::Collapsed;
+	}
+};
 
 
 /*-----------------------------------------------------------------------------
@@ -22,21 +118,30 @@ void SMatineeViewport::Construct(const FArguments& InArgs, TWeakPtr<FMatinee> In
 {
 	this->ChildSlot
 	[
-		SNew(SHorizontalBox)
-		+SHorizontalBox::Slot()
-		.FillWidth(1)
+		SNew(SVerticalBox)
+		+SVerticalBox::Slot()
 		[
-			SAssignNew(ViewportWidget, SViewport)
-			.EnableGammaCorrection(false)
-			.IsEnabled(FSlateApplication::Get().GetNormalExecutionAttribute())
-			.ShowEffectWhenDisabled(false)
+			SNew(SHorizontalBox)
+			+SHorizontalBox::Slot()
+			.FillWidth(1)
+			[
+				SAssignNew(ViewportWidget, SViewport)
+				.EnableGammaCorrection(false)
+				.IsEnabled(FSlateApplication::Get().GetNormalExecutionAttribute())
+				.ShowEffectWhenDisabled(false)
+			]
+			+SHorizontalBox::Slot()
+			.AutoWidth()
+			[
+				SAssignNew(ScrollBar_Vert, SScrollBar)
+				.AlwaysShowScrollbar(true)
+				.OnUserScrolled(this, &SMatineeViewport::OnScroll)
+			]
 		]
-		+SHorizontalBox::Slot()
-		.AutoWidth()
+		+SVerticalBox::Slot()
+		.AutoHeight()
 		[
-			SAssignNew(ScrollBar_Vert, SScrollBar)
-			.AlwaysShowScrollbar(true)
-			.OnUserScrolled(this, &SMatineeViewport::OnScroll)
+			SNew(SMatineeDeprecationMessage)
 		]
 	];
 
