Matinee: Fix to use DarkHyperlink

diff --git a/Engine/Source/Editor/Matinee/Private/MatineeTrackView.cpp b/Engine/Source/Editor/Matinee/Private/MatineeTrackView.cpp
--- a/Engine/Source/Editor/Matinee/Private/MatineeTrackView.cpp
+++ b/Engine/Source/Editor/Matinee/Private/MatineeTrackView.cpp
@@ -61,7 +61,7 @@ public:
 				[
 					SNew(SRichTextBlock)
 					.Text(NSLOCTEXT("Matinee", "MatineeLastVersionSupported", "As of 4.23, Matinee is no longer supported by UE4 and will be removed from the engine in the near future. Once removed, you will <NormalText.Important>no longer be able to run a Matinee or open Matinee Editor</>.\n"
-									"Please use the <a id=\"udn\" href=\"/Engine/Sequencer/HowTo/MatineeConversionTool\" style=\"Hyperlink\">Matinee to Sequencer Conversion Tool</> to convert any files to Sequencer as soon as possible."))
+									"Please use the <a id=\"udn\" href=\"/Engine/Sequencer/HowTo/MatineeConversionTool\" style=\"DarkHyperlink\">Matinee to Sequencer Conversion Tool</> to convert any files to Sequencer as soon as possible."))
 					.AutoWrapText(true)
 					.DecoratorStyleSet(&FEditorStyle::Get())
 					+ SRichTextBlock::HyperlinkDecorator(TEXT("udn"), FSlateHyperlinkRun::FOnClick::CreateStatic(&OnDocLinkClicked))
