Use new inline versions of substring functions
Fix up nearby cases where ESearchCase::CaseSensitive should have been used
#rnx

diff --git a/Engine/Plugins/Developer/CLionSourceCodeAccess/Source/CLionSourceCodeAccess/Private/CLionSourceCodeAccessor.cpp b/Engine/Plugins/Developer/CLionSourceCodeAccess/Source/CLionSourceCodeAccess/Private/CLionSourceCodeAccessor.cpp
--- a/Engine/Plugins/Developer/CLionSourceCodeAccess/Source/CLionSourceCodeAccess/Private/CLionSourceCodeAccessor.cpp
+++ b/Engine/Plugins/Developer/CLionSourceCodeAccess/Source/CLionSourceCodeAccess/Private/CLionSourceCodeAccessor.cpp
@@ -120,13 +120,13 @@ bool FCLionSourceCodeAccessor::OpenSolutionAtPath(const FString& InSolutionPath)
 	FString CorrectSolutionPath = InSolutionPath;
 	if (InSolutionPath.EndsWith(TEXT("UE4")))
 	{
-		CorrectSolutionPath = CorrectSolutionPath.Left(CorrectSolutionPath.Len() - 3);
+		CorrectSolutionPath.LeftInline(CorrectSolutionPath.Len() - 3);
 	}
 	// UE4 passes the project folder and name, so strip the name off
-	int32 LastPathIndex = CorrectSolutionPath.Find(TEXT("/"), ESearchCase::IgnoreCase, ESearchDir::FromEnd);
+	int32 LastPathIndex = CorrectSolutionPath.Find(TEXT("/"), ESearchCase::CaseSensitive, ESearchDir::FromEnd);
 	if (LastPathIndex > -1)
 	{
-		CorrectSolutionPath = CorrectSolutionPath.Left(LastPathIndex + 1);
+		CorrectSolutionPath.LeftInline(LastPathIndex + 1);
 	}
 	// Make sure the path is wrapped in "" properly
 	CorrectSolutionPath = FString::Printf(TEXT("\"%sCMakeLists.txt\""), *CorrectSolutionPath);
diff --git a/Engine/Plugins/Developer/GitSourceControl/Source/GitSourceControl/Private/GitSourceControlUtils.cpp b/Engine/Plugins/Developer/GitSourceControl/Source/GitSourceControl/Private/GitSourceControlUtils.cpp
--- a/Engine/Plugins/Developer/GitSourceControl/Source/GitSourceControl/Private/GitSourceControlUtils.cpp
+++ b/Engine/Plugins/Developer/GitSourceControl/Source/GitSourceControl/Private/GitSourceControlUtils.cpp
@@ -427,7 +427,7 @@ bool FindRootDirectory(const FString& InPath, FString& OutRepositoryRoot)
 		int32 Len = Str.Len();
 		while(Len && Str[Len - 1] == Char)
 		{
-			Str = Str.LeftChop(1);
+			Str.LeftChopInline(1);
 			Len = Str.Len();
 		}
 	};
@@ -445,7 +445,7 @@ bool FindRootDirectory(const FString& InPath, FString& OutRepositoryRoot)
 			int32 LastSlashIndex;
 			if(OutRepositoryRoot.FindLastChar('/', LastSlashIndex))
 			{
-				OutRepositoryRoot = OutRepositoryRoot.Left(LastSlashIndex);
+				OutRepositoryRoot.LeftInline(LastSlashIndex);
 			}
 			else
 			{
diff --git a/Engine/Plugins/Developer/PerforceSourceControl/Source/PerforceSourceControl/Private/PerforceSourceControlOperations.cpp b/Engine/Plugins/Developer/PerforceSourceControl/Source/PerforceSourceControl/Private/PerforceSourceControlOperations.cpp
--- a/Engine/Plugins/Developer/PerforceSourceControl/Source/PerforceSourceControl/Private/PerforceSourceControlOperations.cpp
+++ b/Engine/Plugins/Developer/PerforceSourceControl/Source/PerforceSourceControl/Private/PerforceSourceControlOperations.cpp
@@ -998,7 +998,7 @@ static void ParseOpenedResults(const FP4RecordSet& InRecords, const FString& Cli
 			{
 				// Null clients use the pattern in PathRoot: //Workspace/FileName
 				// Here we chop off the '//Workspace/' to return the workspace filename
-				FullPath = FullPath.RightChop(PathRoot.Len() + 1);
+				FullPath.RightChopInline(PathRoot.Len() + 1, false);
 			}
 			else
 			{
diff --git a/Engine/Plugins/Editor/AssetManagerEditor/Source/AssetManagerEditor/Private/AssetManagerEditorModule.cpp b/Engine/Plugins/Editor/AssetManagerEditor/Source/AssetManagerEditor/Private/AssetManagerEditorModule.cpp
--- a/Engine/Plugins/Editor/AssetManagerEditor/Source/AssetManagerEditor/Private/AssetManagerEditorModule.cpp
+++ b/Engine/Plugins/Editor/AssetManagerEditor/Source/AssetManagerEditor/Private/AssetManagerEditorModule.cpp
@@ -752,7 +752,7 @@ void FAssetManagerEditorModule::OpenShaderCookStatistics(TArray<FName> SelectedP
 			}
 			CommonIdentical = FMath::Min(Identical, CommonIdentical);
 		}
-		CommonSubPath = CommonPath.Left(CommonIdentical);
+		CommonPath.LeftInline(CommonIdentical);
 	}
 	static const FName LevelEditorModuleName("LevelEditor");
 	static const FName LevelEditorStatsViewerTab("LevelEditorStatsViewer");
diff --git a/Engine/Plugins/Editor/GameplayTagsEditor/Source/GameplayTagsEditor/Private/SGameplayTagContainerGraphPin.cpp b/Engine/Plugins/Editor/GameplayTagsEditor/Source/GameplayTagsEditor/Private/SGameplayTagContainerGraphPin.cpp
--- a/Engine/Plugins/Editor/GameplayTagsEditor/Source/GameplayTagsEditor/Private/SGameplayTagContainerGraphPin.cpp
+++ b/Engine/Plugins/Editor/GameplayTagsEditor/Source/GameplayTagsEditor/Private/SGameplayTagContainerGraphPin.cpp
@@ -46,29 +46,29 @@ void SGameplayTagContainerGraphPin::ParseDefaultValueData()
 
 	FilterString = GameplayTagPinUtilities::ExtractTagFilterStringFromGraphPin(GraphPinObj);
 
-	if (TagString.StartsWith(TEXT("(")) && TagString.EndsWith(TEXT(")")))
+	if (TagString.StartsWith(TEXT("("), ESearchCase::CaseSensitive) && TagString.EndsWith(TEXT(")"), ESearchCase::CaseSensitive))
 	{
-		TagString = TagString.LeftChop(1);
-		TagString = TagString.RightChop(1);
+		TagString.LeftChopInline(1, false);
+		TagString.RightChopInline(1, false);
 
-		TagString.Split("=", NULL, &TagString);
+		TagString.Split(TEXT("="), nullptr, &TagString, ESearchCase::CaseSensitive);
 
-		TagString = TagString.LeftChop(1);
-		TagString = TagString.RightChop(1);
+		TagString.LeftChopInline(1, false);
+		TagString.RightChopInline(1, false);
 
 		FString ReadTag;
 		FString Remainder;
 
-		while (TagString.Split(TEXT(","), &ReadTag, &Remainder))
+		while (TagString.Split(TEXT(","), &ReadTag, &Remainder, ESearchCase::CaseSensitive))
 		{
-			ReadTag.Split("=", NULL, &ReadTag);
-			if (ReadTag.EndsWith(TEXT(")")))
+			ReadTag.Split(TEXT("="), NULL, &ReadTag, ESearchCase::CaseSensitive);
+			if (ReadTag.EndsWith(TEXT(")"), ESearchCase::CaseSensitive))
 			{
-				ReadTag = ReadTag.LeftChop(1);
-				if (ReadTag.StartsWith(TEXT("\"")) && ReadTag.EndsWith(TEXT("\"")))
+				ReadTag.LeftChopInline(1, false);
+				if (ReadTag.StartsWith(TEXT("\""), ESearchCase::CaseSensitive) && ReadTag.EndsWith(TEXT("\""), ESearchCase::CaseSensitive))
 				{
-					ReadTag = ReadTag.LeftChop(1);
-					ReadTag = ReadTag.RightChop(1);
+					ReadTag.LeftChopInline(1, false);
+					ReadTag.RightChopInline(1, false);
 				}
 			}
 			TagString = Remainder;
@@ -81,14 +81,14 @@ void SGameplayTagContainerGraphPin::ParseDefaultValueData()
 		}
 		if (!Remainder.IsEmpty())
 		{
-			Remainder.Split("=", NULL, &Remainder);
-			if (Remainder.EndsWith(TEXT(")")))
+			Remainder.Split(TEXT("="), nullptr, &Remainder, ESearchCase::CaseSensitive);
+			if (Remainder.EndsWith(TEXT(")"), ESearchCase::CaseSensitive))
 			{
-				Remainder = Remainder.LeftChop(1);
-				if (Remainder.StartsWith(TEXT("\"")) && Remainder.EndsWith(TEXT("\"")))
+				Remainder.LeftChopInline(1, false);
+				if (Remainder.StartsWith(TEXT("\""), ESearchCase::CaseSensitive) && Remainder.EndsWith(TEXT("\""), ESearchCase::CaseSensitive))
 				{
-					Remainder = Remainder.LeftChop(1);
-					Remainder = Remainder.RightChop(1);
+					Remainder.LeftChopInline(1, false);
+					Remainder.RightChopInline(1, false);
 				}
 			}
 			FGameplayTag GameplayTag = FGameplayTag::RequestGameplayTag(FName(*Remainder));
diff --git a/Engine/Plugins/Editor/GameplayTagsEditor/Source/GameplayTagsEditor/Private/SGameplayTagGraphPin.cpp b/Engine/Plugins/Editor/GameplayTagsEditor/Source/GameplayTagsEditor/Private/SGameplayTagGraphPin.cpp
--- a/Engine/Plugins/Editor/GameplayTagsEditor/Source/GameplayTagsEditor/Private/SGameplayTagGraphPin.cpp
+++ b/Engine/Plugins/Editor/GameplayTagsEditor/Source/GameplayTagsEditor/Private/SGameplayTagGraphPin.cpp
@@ -47,15 +47,15 @@ void SGameplayTagGraphPin::ParseDefaultValueData()
 
 	FilterString = GameplayTagPinUtilities::ExtractTagFilterStringFromGraphPin(GraphPinObj);
 
-	if (TagString.StartsWith(TEXT("(")) && TagString.EndsWith(TEXT(")")))
+	if (TagString.StartsWith(TEXT("("), ESearchCase::CaseSensitive) && TagString.EndsWith(TEXT(")"), ESearchCase::CaseSensitive))
 	{
-		TagString = TagString.LeftChop(1);
-		TagString = TagString.RightChop(1);
-		TagString.Split("=", NULL, &TagString);
-		if (TagString.StartsWith(TEXT("\"")) && TagString.EndsWith(TEXT("\"")))
+		TagString.LeftChopInline(1, false);
+		TagString.RightChopInline(1, false);
+		TagString.Split(TEXT("="), nullptr, &TagString, ESearchCase::CaseSensitive);
+		if (TagString.StartsWith(TEXT("\""), ESearchCase::CaseSensitive) && TagString.EndsWith(TEXT("\""), ESearchCase::CaseSensitive))
 		{
-			TagString = TagString.LeftChop(1);
-			TagString = TagString.RightChop(1);
+			TagString.LeftChopInline(1, false);
+			TagString.RightChopInline(1, false);
 		}
 	}
 
diff --git a/Engine/Plugins/Enterprise/VariantManagerContent/Source/VariantManagerContent/Private/LevelVariantSets.cpp b/Engine/Plugins/Enterprise/VariantManagerContent/Source/VariantManagerContent/Private/LevelVariantSets.cpp
--- a/Engine/Plugins/Enterprise/VariantManagerContent/Source/VariantManagerContent/Private/LevelVariantSets.cpp
+++ b/Engine/Plugins/Enterprise/VariantManagerContent/Source/VariantManagerContent/Private/LevelVariantSets.cpp
@@ -178,7 +178,7 @@ FString ULevelVariantSets::GetUniqueVariantSetName(const FString& InPrefix)
 	FString LastChar = VarSetName.Right(1);
 	while (LastChar.IsNumeric())
 	{
-		VarSetName = VarSetName.LeftChop(1);
+		VarSetName.LeftChopInline(1, false);
 		LastChar = VarSetName.Right(1);
 	}
 
diff --git a/Engine/Plugins/Enterprise/VariantManagerContent/Source/VariantManagerContent/Private/VariantSet.cpp b/Engine/Plugins/Enterprise/VariantManagerContent/Source/VariantManagerContent/Private/VariantSet.cpp
--- a/Engine/Plugins/Enterprise/VariantManagerContent/Source/VariantManagerContent/Private/VariantSet.cpp
+++ b/Engine/Plugins/Enterprise/VariantManagerContent/Source/VariantManagerContent/Private/VariantSet.cpp
@@ -88,7 +88,7 @@ FString UVariantSet::GetUniqueVariantName(const FString& InPrefix)
 	FString LastChar = VarName.Right(1);
 	while (LastChar.IsNumeric())
 	{
-		VarName = VarName.LeftChop(1);
+		VarName.LeftChopInline(1, false);
 		LastChar = VarName.Right(1);
 	}
 
diff --git a/Engine/Plugins/Experimental/AutomationUtils/Source/AutomationUtils/Private/AutomationUtilsBlueprintLibrary.cpp b/Engine/Plugins/Experimental/AutomationUtils/Source/AutomationUtils/Private/AutomationUtilsBlueprintLibrary.cpp
--- a/Engine/Plugins/Experimental/AutomationUtils/Source/AutomationUtils/Private/AutomationUtilsBlueprintLibrary.cpp
+++ b/Engine/Plugins/Experimental/AutomationUtils/Source/AutomationUtils/Private/AutomationUtilsBlueprintLibrary.cpp
@@ -210,7 +210,7 @@ FAutomationUtilsGameplayAutomationScreenshotInstance::FAutomationUtilsGameplayAu
 		if (HardwareDetailsString.Len() > 0)
 		{
 			//remove leading "_"
-			HardwareDetailsString = HardwareDetailsString.RightChop(1);
+			HardwareDetailsString.RightChopInline(1, false);
 		}
 
 		//now plop that back onto the path we're building
diff --git a/Engine/Plugins/Experimental/BackChannel/Source/BackChannel/Private/Protocol/BackChannelOSCMessage.cpp b/Engine/Plugins/Experimental/BackChannel/Source/BackChannel/Private/Protocol/BackChannelOSCMessage.cpp
--- a/Engine/Plugins/Experimental/BackChannel/Source/BackChannel/Private/Protocol/BackChannelOSCMessage.cpp
+++ b/Engine/Plugins/Experimental/BackChannel/Source/BackChannel/Private/Protocol/BackChannelOSCMessage.cpp
@@ -194,7 +194,7 @@ TSharedPtr<FBackChannelOSCMessage> FBackChannelOSCMessage::CreateFromBuffer(cons
 	FString TagString = ANSI_TO_TCHAR((ANSICHAR*)pParsedData);
 
 	// remove the leading ,
-	NewMessage->TagString = TagString.RightChop(TagString.Find(TEXT(","))+1);
+	NewMessage->TagString = TagString.RightChop(TagString.Find(TEXT(","), ESearchCase::CaseSensitive)+1);
 
 	pParsedData += RoundedArgumentSize(TagString.Len() + 1);
 
diff --git a/Engine/Plugins/Experimental/ControlRig/Source/ControlRigDeveloper/Private/ControlRigBlueprintUtils.cpp b/Engine/Plugins/Experimental/ControlRig/Source/ControlRigDeveloper/Private/ControlRigBlueprintUtils.cpp
--- a/Engine/Plugins/Experimental/ControlRig/Source/ControlRigDeveloper/Private/ControlRigBlueprintUtils.cpp
+++ b/Engine/Plugins/Experimental/ControlRig/Source/ControlRigDeveloper/Private/ControlRigBlueprintUtils.cpp
@@ -66,7 +66,7 @@ FName FControlRigBlueprintUtils::ValidateName(UBlueprint* InBlueprint, const FSt
 	FString Name = InName;
 	if (Name.StartsWith(TEXT("RigUnit_")))
 	{
-		Name = Name.RightChop(8);
+		Name.RightChopInline(8, false);
 	}
 
 	TSharedPtr<FKismetNameValidator> NameValidator;
@@ -106,7 +106,7 @@ FName FControlRigBlueprintUtils::ValidateName(UBlueprint* InBlueprint, const FSt
 		// If the length of the final string will be too long, cut off the end so we can fit the number
 		if (CountLength + BaseName.Len() > NameValidator->GetMaximumNameLength())
 		{
-			BaseName = BaseName.Left(NameValidator->GetMaximumNameLength() - CountLength);
+			BaseName.LeftInline(NameValidator->GetMaximumNameLength() - CountLength);
 		}
 		Name = FString::Printf(TEXT("%s_%d"), *BaseName, Count);
 		Count++;
diff --git a/Engine/Plugins/FX/Niagara/Source/NiagaraEditor/Private/NiagaraHlslTranslator.cpp b/Engine/Plugins/FX/Niagara/Source/NiagaraEditor/Private/NiagaraHlslTranslator.cpp
--- a/Engine/Plugins/FX/Niagara/Source/NiagaraEditor/Private/NiagaraHlslTranslator.cpp
+++ b/Engine/Plugins/FX/Niagara/Source/NiagaraEditor/Private/NiagaraHlslTranslator.cpp
@@ -4554,7 +4554,7 @@ void FHlslNiagaraTranslator::FunctionCall(UNiagaraNodeFunctionCall* FunctionNode
 				FString VarStr = Var.GetName().ToString();
 				if (VarStr.StartsWith(*ModuleAlias))
 				{
-					VarStr = VarStr.Mid(ModuleAlias->Len() + 1);
+					VarStr.MidInline(ModuleAlias->Len() + 1, MAX_int32, false);
 					if (FunctionNode->FindStaticSwitchInputPin(*VarStr))
 					{
 						Error(FText::Format(LOCTEXT("SwitchPinFoundForSetPin", "A switch node pin exists but is being set directly using Set node! Please use the stack UI to resolve the conflict. Output Pin: {0}"), FText::FromName(Var.GetName())), FunctionNode, nullptr);
diff --git a/Engine/Plugins/NetcodeUnitTest/NetcodeUnitTest/Source/NetcodeUnitTest/Private/NUTUtil.cpp b/Engine/Plugins/NetcodeUnitTest/NetcodeUnitTest/Source/NetcodeUnitTest/Private/NUTUtil.cpp
--- a/Engine/Plugins/NetcodeUnitTest/NetcodeUnitTest/Source/NetcodeUnitTest/Private/NUTUtil.cpp
+++ b/Engine/Plugins/NetcodeUnitTest/NetcodeUnitTest/Source/NetcodeUnitTest/Private/NUTUtil.cpp
@@ -136,7 +136,7 @@ bool NUTUtil::ParseEnum(const TCHAR* Stream, const TCHAR* Match, const TCHAR* En
 					PrefixIdx++;
 				}
 
-				Prefix = Prefix.Left(PrefixIdx);
+				Prefix.LeftInline(PrefixIdx);
 			}
 
 
@@ -144,7 +144,14 @@ bool NUTUtil::ParseEnum(const TCHAR* Stream, const TCHAR* Match, const TCHAR* En
 			{
 				int32 UnderscoreIdx = Prefix.Find(TEXT("_"), ESearchCase::CaseSensitive, ESearchDir::FromEnd);
 
-				Prefix = UnderscoreIdx != INDEX_NONE ? Prefix.Left(UnderscoreIdx+1) : TEXT("");
+				if (UnderscoreIdx != INDEX_NONE)
+				{
+					Prefix.LeftInline(UnderscoreIdx+1);
+				}
+				else
+				{
+					Prefix.Reset();
+				}
 
 				if (Prefix.Len() > 0 && !StrValue.StartsWith(Prefix))
 				{
diff --git a/Engine/Plugins/Online/Android/OnlineSubsystemGooglePlay/Source/Private/OnlineStoreGooglePlay.cpp b/Engine/Plugins/Online/Android/OnlineSubsystemGooglePlay/Source/Private/OnlineStoreGooglePlay.cpp
--- a/Engine/Plugins/Online/Android/OnlineSubsystemGooglePlay/Source/Private/OnlineStoreGooglePlay.cpp
+++ b/Engine/Plugins/Online/Android/OnlineSubsystemGooglePlay/Source/Private/OnlineStoreGooglePlay.cpp
@@ -59,7 +59,8 @@ TSharedRef<FOnlineStoreOffer> ConvertProductToStoreOffer(const FInAppPurchasePro
 	int32 CloseParenIdx = -1;
 	if (Title.FindLastChar(TEXT(')'), CloseParenIdx) && Title.FindLastChar(TEXT('('), OpenParenIdx) && (OpenParenIdx < CloseParenIdx))
 	{
-		Title = Title.Left(OpenParenIdx).TrimEnd();
+		Title.LeftInline(OpenParenIdx);
+		Title.TrimEndInline();
 	}
 
 	NewProductInfo->Title = FText::FromString(Title);
diff --git a/Engine/Plugins/Online/OnlineSubsystem/Source/Private/OnlineSubsystemModule.cpp b/Engine/Plugins/Online/OnlineSubsystem/Source/Private/OnlineSubsystemModule.cpp
--- a/Engine/Plugins/Online/OnlineSubsystem/Source/Private/OnlineSubsystemModule.cpp
+++ b/Engine/Plugins/Online/OnlineSubsystem/Source/Private/OnlineSubsystemModule.cpp
@@ -119,13 +119,13 @@ static inline void ReadOnlineSubsystemConfigPairs(const TCHAR* Section, const TC
 
 		if (TrimmedConfigEntry.Left(1) == TEXT("("))
 		{
-			TrimmedConfigEntry = TrimmedConfigEntry.RightChop(1);
+			TrimmedConfigEntry.RightChopInline(1, false);
 		}
 		if (TrimmedConfigEntry.Right(1) == TEXT(")"))
 		{
-			TrimmedConfigEntry = TrimmedConfigEntry.LeftChop(1);
+			TrimmedConfigEntry.LeftChopInline(1, false);
 		}
-		if (TrimmedConfigEntry.Split(TEXT("="), &KeyString, &ValueString))
+		if (TrimmedConfigEntry.Split(TEXT("="), &KeyString, &ValueString, ESearchCase::CaseSensitive))
 		{
 			KeyString.TrimStartAndEndInline();
 			ValueString.TrimStartAndEndInline();
diff --git a/Engine/Plugins/Online/OnlineSubsystemSteam/Source/Private/OnlineSessionInterfaceSteam.cpp b/Engine/Plugins/Online/OnlineSubsystemSteam/Source/Private/OnlineSessionInterfaceSteam.cpp
--- a/Engine/Plugins/Online/OnlineSubsystemSteam/Source/Private/OnlineSessionInterfaceSteam.cpp
+++ b/Engine/Plugins/Online/OnlineSubsystemSteam/Source/Private/OnlineSessionInterfaceSteam.cpp
@@ -1346,7 +1346,7 @@ FString FOnlineSessionSteam::GetCustomDedicatedServerName() const
 		{
 			UE_LOG_ONLINE_SESSION(Warning, TEXT("SteamServerName overflows the maximum amount of characters %d allowed, truncating."), k_cbMaxGameServerName);
 			// Must have space for the null terminator
-			ServerName = ServerName.Left(k_cbMaxGameServerName - 1);
+			ServerName.LeftInline(k_cbMaxGameServerName - 1);
 		}
 
 		return ServerName;
diff --git a/Engine/Plugins/Online/VoiceChat/VivoxVoiceChat/Source/Private/VivoxVoiceChat.cpp b/Engine/Plugins/Online/VoiceChat/VivoxVoiceChat/Source/Private/VivoxVoiceChat.cpp
--- a/Engine/Plugins/Online/VoiceChat/VivoxVoiceChat/Source/Private/VivoxVoiceChat.cpp
+++ b/Engine/Plugins/Online/VoiceChat/VivoxVoiceChat/Source/Private/VivoxVoiceChat.cpp
@@ -1827,7 +1827,7 @@ FString FVivoxVoiceChat::GetChannelNameFromUri(const VivoxClientApi::Uri& Channe
 		int32 Channel3dParametersIndex = Channel.Find(TEXT("!p-"));
 		if (Channel3dParametersIndex != INDEX_NONE)
 		{
-			Channel = Channel.Left(Channel3dParametersIndex);
+			Channel.LeftInline(Channel3dParametersIndex);
 		}
 		return Channel;
 	}
diff --git a/Engine/Plugins/Runtime/GameplayAbilities/Source/GameplayAbilitiesEditor/Private/GameplayCueTagDetails.cpp b/Engine/Plugins/Runtime/GameplayAbilities/Source/GameplayAbilitiesEditor/Private/GameplayCueTagDetails.cpp
--- a/Engine/Plugins/Runtime/GameplayAbilities/Source/GameplayAbilitiesEditor/Private/GameplayCueTagDetails.cpp
+++ b/Engine/Plugins/Runtime/GameplayAbilities/Source/GameplayAbilitiesEditor/Private/GameplayCueTagDetails.cpp
@@ -115,7 +115,7 @@ TSharedRef<ITableRow> FGameplayCueTagDetails::GenerateListRow(TSharedRef<FSoftOb
 	int32 idx;
 	if (ShortName.FindLastChar(TEXT('.'), idx))
 	{
-		ShortName = ShortName.Right(ShortName.Len() - (idx+1));
+		ShortName.RightInline(ShortName.Len() - (idx+1), false);
 		ShortName.RemoveFromEnd(TEXT("_c"));
 	}
 
diff --git a/Engine/Plugins/Runtime/GameplayAbilities/Source/GameplayAbilitiesEditor/Private/SGameplayCueEditor.cpp b/Engine/Plugins/Runtime/GameplayAbilities/Source/GameplayAbilitiesEditor/Private/SGameplayCueEditor.cpp
--- a/Engine/Plugins/Runtime/GameplayAbilities/Source/GameplayAbilitiesEditor/Private/SGameplayCueEditor.cpp
+++ b/Engine/Plugins/Runtime/GameplayAbilities/Source/GameplayAbilitiesEditor/Private/SGameplayCueEditor.cpp
@@ -903,10 +903,10 @@ TSharedRef<ITableRow> SGameplayCueEditorImpl::OnGenerateWidgetForGameplayCueList
 					int32 idx;
 					if (ObjName.FindLastChar(TEXT('.'), idx))
 					{
-						ObjName = ObjName.RightChop(idx + 1);
+						ObjName.RightChopInline(idx + 1, false);
 						if (ObjName.FindLastChar(TEXT('_'), idx))
 						{
-							ObjName = ObjName.Left(idx);
+							ObjName.LeftInline(idx);
 						}
 					}
 
diff --git a/Engine/Plugins/Runtime/PhysXVehicles/Source/PhysXVehiclesEditor/Private/PhysXVehiclesEditorPlugin.cpp b/Engine/Plugins/Runtime/PhysXVehicles/Source/PhysXVehiclesEditor/Private/PhysXVehiclesEditorPlugin.cpp
--- a/Engine/Plugins/Runtime/PhysXVehicles/Source/PhysXVehiclesEditor/Private/PhysXVehiclesEditorPlugin.cpp
+++ b/Engine/Plugins/Runtime/PhysXVehicles/Source/PhysXVehiclesEditor/Private/PhysXVehiclesEditorPlugin.cpp
@@ -63,7 +63,7 @@ static void ConvertTireTypes()
 			int32 SlashPos = INDEX_NONE;
 			if (TirePath.FindLastChar('/', SlashPos))
 			{
-				TirePath = TirePath.Left(SlashPos);
+				TirePath.LeftInline(SlashPos);
 
 				FString TireTypeName = TireType->GetName();
 				FString TireConfigName = FString::Printf(TEXT("%s_TireConfig"), *TireTypeName);
diff --git a/Engine/Plugins/Runtime/ReplicationGraph/Source/Public/ReplicationGraphTypes.h b/Engine/Plugins/Runtime/ReplicationGraph/Source/Public/ReplicationGraphTypes.h
--- a/Engine/Plugins/Runtime/ReplicationGraph/Source/Public/ReplicationGraphTypes.h
+++ b/Engine/Plugins/Runtime/ReplicationGraph/Source/Public/ReplicationGraphTypes.h
@@ -1282,7 +1282,7 @@ struct FReplicationGraphDebugInfo
 	void Log(const FString& Str) { Ar.Logf(TEXT("%s%s"), *CurrentIndentString, *Str); }
 
 	void PushIndent() { CurrentIndentString += IndentString; }
-	void PopIndent() { CurrentIndentString = CurrentIndentString.LeftChop(IndentString.Len()); }
+	void PopIndent() { CurrentIndentString.LeftChopInline(IndentString.Len(), false); }
 };
 
 REPLICATIONGRAPH_API void LogActorRepList(FReplicationGraphDebugInfo& DebugInfo, FString Prefix, const FActorRepListRefView& List);
diff --git a/Engine/Plugins/Runtime/WindowsMixedReality/Source/WindowsMixedRealityHandTracking/Private/WindowsMixedRealityHandTracking.cpp b/Engine/Plugins/Runtime/WindowsMixedReality/Source/WindowsMixedRealityHandTracking/Private/WindowsMixedRealityHandTracking.cpp
--- a/Engine/Plugins/Runtime/WindowsMixedReality/Source/WindowsMixedRealityHandTracking/Private/WindowsMixedRealityHandTracking.cpp
+++ b/Engine/Plugins/Runtime/WindowsMixedReality/Source/WindowsMixedRealityHandTracking/Private/WindowsMixedRealityHandTracking.cpp
@@ -207,12 +207,12 @@ bool FWindowsMixedRealityHandTracking::GetControllerOrientationAndPosition(const
 			{
 				bUseRightHand = true;
 				// Strip off the Right
-				SourceString = SourceString.Right(SourceString.Len() - 5);
+				SourceString.RightInline(SourceString.Len() - 5, false);
 			}
 			else
 			{
 				// Strip off the Left
-				SourceString = SourceString.Right(SourceString.Len() - 4);
+				SourceString.RightInline(SourceString.Len() - 4, false);
 			}
 			FName FullEnumName(*FString(TEXT("EWMRHandKeypoint::") + SourceString), FNAME_Find);
 			// Get the enum value from the name
diff --git a/Engine/Plugins/ScriptPlugin/Source/ScriptGeneratorPlugin/Private/ScriptCodeGeneratorBase.cpp b/Engine/Plugins/ScriptPlugin/Source/ScriptGeneratorPlugin/Private/ScriptCodeGeneratorBase.cpp
--- a/Engine/Plugins/ScriptPlugin/Source/ScriptGeneratorPlugin/Private/ScriptCodeGeneratorBase.cpp
+++ b/Engine/Plugins/ScriptPlugin/Source/ScriptGeneratorPlugin/Private/ScriptCodeGeneratorBase.cpp
@@ -72,27 +72,27 @@ FString FScriptCodeGeneratorBase::GetClassNameCPP(UClass* Class)
 
 FString FScriptCodeGeneratorBase::GetPropertyTypeCPP(UProperty* Property, uint32 PortFlags /*= 0*/)
 {
-	static FString EnumDecl(TEXT("enum "));
-	static FString StructDecl(TEXT("struct "));
-	static FString ClassDecl(TEXT("class "));
-	static FString TEnumAsByteDecl(TEXT("TEnumAsByte<enum "));
-	static FString TSubclassOfDecl(TEXT("TSubclassOf<class "));
+	static const FString EnumDecl(TEXT("enum "));
+	static const FString StructDecl(TEXT("struct "));
+	static const FString ClassDecl(TEXT("class "));
+	static const FString TEnumAsByteDecl(TEXT("TEnumAsByte<enum "));
+	static const FString TSubclassOfDecl(TEXT("TSubclassOf<class "));
 
 	FString PropertyType = Property->GetCPPType(NULL, PortFlags);
 	// Strip any forward declaration keywords
-	if (PropertyType.StartsWith(EnumDecl) || PropertyType.StartsWith(StructDecl) || PropertyType.StartsWith(ClassDecl))
+	if (PropertyType.StartsWith(EnumDecl, ESearchCase::CaseSensitive) || PropertyType.StartsWith(StructDecl, ESearchCase::CaseSensitive) || PropertyType.StartsWith(ClassDecl, ESearchCase::CaseSensitive))
 	{
-		int FirstSpaceIndex = PropertyType.Find(TEXT(" "));
-		PropertyType = PropertyType.Mid(FirstSpaceIndex + 1);
+		int32 FirstSpaceIndex = PropertyType.Find(TEXT(" "), ESearchCase::CaseSensitive);
+		PropertyType.MidInline(FirstSpaceIndex + 1, MAX_int32, false);
 	}
-	else if (PropertyType.StartsWith(TEnumAsByteDecl))
+	else if (PropertyType.StartsWith(TEnumAsByteDecl, ESearchCase::CaseSensitive))
 	{
-		int FirstSpaceIndex = PropertyType.Find(TEXT(" "));
+		int32 FirstSpaceIndex = PropertyType.Find(TEXT(" "), ESearchCase::CaseSensitive);
 		PropertyType = TEXT("TEnumAsByte<") + PropertyType.Mid(FirstSpaceIndex + 1);
 	}
-	else if (PropertyType.StartsWith(TSubclassOfDecl))
+	else if (PropertyType.StartsWith(TSubclassOfDecl), ESearchCase::CaseSensitive)
 	{
-		int FirstSpaceIndex = PropertyType.Find(TEXT(" "));
+		int32 FirstSpaceIndex = PropertyType.Find(TEXT(" "), ESearchCase::CaseSensitive);
 		PropertyType = TEXT("TSubclassOf<") + PropertyType.Mid(FirstSpaceIndex + 1);
 	}
 	return PropertyType;
diff --git a/Engine/Plugins/Tests/EditorTests/Source/EditorTests/Private/UnrealEd/EditorAssetAutomationTests.cpp b/Engine/Plugins/Tests/EditorTests/Source/EditorTests/Private/UnrealEd/EditorAssetAutomationTests.cpp
--- a/Engine/Plugins/Tests/EditorTests/Source/EditorTests/Private/UnrealEd/EditorAssetAutomationTests.cpp
+++ b/Engine/Plugins/Tests/EditorTests/Source/EditorTests/Private/UnrealEd/EditorAssetAutomationTests.cpp
@@ -1003,9 +1003,9 @@ namespace ImportExportAssetHelper
 					Extension = FPaths::GetExtension(ImportPath);
 				}
 
-				if (Extension.StartsWith(TEXT(".")))
+				if (Extension.StartsWith(TEXT("."), ESearchCase::CaseSensitive))
 				{
-					Extension = Extension.RightChop(1);
+					Extension.RightChopInline(1, false);
 				}
 
 				//Export the asset
diff --git a/Engine/Source/Developer/AssetTools/Private/AssetTools.cpp b/Engine/Source/Developer/AssetTools/Private/AssetTools.cpp
--- a/Engine/Source/Developer/AssetTools/Private/AssetTools.cpp
+++ b/Engine/Source/Developer/AssetTools/Private/AssetTools.cpp
@@ -2277,7 +2277,7 @@ void UAssetToolsImpl::ExportAssetsInternal(const TArray<UObject*>& ObjectsToExpo
 				if (PackageName.Left(1) == TEXT("/"))
 				{
 					// Trim the leading slash so the file manager doesn't get confused
-					PackageName = PackageName.Mid(1);
+					PackageName.MidInline(1, false);
 				}
 
 				FPaths::NormalizeFilename(PackageName);
diff --git a/Engine/Source/Developer/AutomationController/Private/AutomationDeviceClusterManager.cpp b/Engine/Source/Developer/AutomationController/Private/AutomationDeviceClusterManager.cpp
--- a/Engine/Source/Developer/AutomationController/Private/AutomationDeviceClusterManager.cpp
+++ b/Engine/Source/Developer/AutomationController/Private/AutomationDeviceClusterManager.cpp
@@ -104,7 +104,7 @@ FString FAutomationDeviceClusterManager::GetGroupNameForDevice(const FDeviceStat
 	if( OutGroupName.Len() > 0 )
 	{
 		//Get rid of the trailing '-'
-		OutGroupName = OutGroupName.LeftChop(1);
+		OutGroupName.LeftChopInline(1, false);
 	}
 
 	return OutGroupName;
diff --git a/Engine/Source/Developer/AutomationDriver/Private/Locators/SlateWidgetLocatorByPath.cpp b/Engine/Source/Developer/AutomationDriver/Private/Locators/SlateWidgetLocatorByPath.cpp
--- a/Engine/Source/Developer/AutomationDriver/Private/Locators/SlateWidgetLocatorByPath.cpp
+++ b/Engine/Source/Developer/AutomationDriver/Private/Locators/SlateWidgetLocatorByPath.cpp
@@ -264,7 +264,7 @@ private:
 			while (TempPath.FindChar(TEXT('/'), Index))
 			{
 				const FString PathPiece = TempPath.Left(Index);
-				TempPath = TempPath.RightChop(Index + 1);
+				TempPath.RightChopInline(Index + 1, false);
 
 				if (PathPiece.Len() == 0)
 				{
diff --git a/Engine/Source/Developer/AutomationWindow/Private/SAutomationWindow.cpp b/Engine/Source/Developer/AutomationWindow/Private/SAutomationWindow.cpp
--- a/Engine/Source/Developer/AutomationWindow/Private/SAutomationWindow.cpp
+++ b/Engine/Source/Developer/AutomationWindow/Private/SAutomationWindow.cpp
@@ -1667,7 +1667,7 @@ TSharedRef<ITableRow> SAutomationWindow::OnGenerateWidgetForLog(TSharedPtr<FAuto
 		int32 LineNumber = FCString::Atoi(*FileAndLineRegexMatcher.GetCaptureGroup(2));
 
 		// Remove the hyperlink from the message, since we're splitting it into its own string.
-		MessageString = MessageString.RightChop(FileAndLineRegexMatcher.GetMatchEnding());
+		MessageString.RightChopInline(FileAndLineRegexMatcher.GetMatchEnding(), false);
 
 		SourceLink = SNew(SHyperlink)
 			.Style(FEditorStyle::Get(), "Common.GotoNativeCodeHyperlink")
diff --git a/Engine/Source/Developer/DerivedDataCache/Private/DerivedDataCache.cpp b/Engine/Source/Developer/DerivedDataCache/Private/DerivedDataCache.cpp
--- a/Engine/Source/Developer/DerivedDataCache/Private/DerivedDataCache.cpp
+++ b/Engine/Source/Developer/DerivedDataCache/Private/DerivedDataCache.cpp
@@ -960,7 +960,7 @@ public:
 			FString Name((*Iter)->GetName());
 			if (Name.Len() > 1024)
 			{
-				Name = Name.Left(1024);
+				Name.LeftInline(1024);
 				Name += TEXT("...");
 			}
 			UE_LOG(LogDerivedDataCache, Warning, TEXT("Leaked Rollup! %s"), *Name);
diff --git a/Engine/Source/Developer/DesktopPlatform/Private/Windows/DesktopPlatformWindows.cpp b/Engine/Source/Developer/DesktopPlatform/Private/Windows/DesktopPlatformWindows.cpp
--- a/Engine/Source/Developer/DesktopPlatform/Private/Windows/DesktopPlatformWindows.cpp
+++ b/Engine/Source/Developer/DesktopPlatform/Private/Windows/DesktopPlatformWindows.cpp
@@ -257,7 +257,7 @@ bool FDesktopPlatformWindows::FileDialogShared(bool bSave, const void* ParentWin
 								int32 WildCardIndex = INDEX_NONE;
 								if (CleanExtension.FindChar(TEXT('*'), WildCardIndex))
 								{
-									CleanExtension = CleanExtension.RightChop(WildCardIndex + 1);
+									CleanExtension.RightChopInline(WildCardIndex + 1, false);
 								}
 							}
 
diff --git a/Engine/Source/Developer/HotReload/Private/HotReload.cpp b/Engine/Source/Developer/HotReload/Private/HotReload.cpp
--- a/Engine/Source/Developer/HotReload/Private/HotReload.cpp
+++ b/Engine/Source/Developer/HotReload/Private/HotReload.cpp
@@ -1346,14 +1346,14 @@ void FHotReloadModule::StripModuleSuffixFromFilename(FString& InOutModuleFilenam
 		int32 SecondHyphenIndex = FirstHyphenIndex;
 		do
 		{
-			SecondHyphenIndex = InOutModuleFilename.Find(TEXT("-"), ESearchCase::IgnoreCase, ESearchDir::FromStart, SecondHyphenIndex + 1);
+			SecondHyphenIndex = InOutModuleFilename.Find(TEXT("-"), ESearchCase::CaseSensitive, ESearchDir::FromStart, SecondHyphenIndex + 1);
 			if (SecondHyphenIndex != INDEX_NONE)
 			{
 				// Make sure that the section between hyphens is the expected module name. This guards against cases where module name has a hyphen inside.
 				FString HotReloadedModuleName = InOutModuleFilename.Mid(FirstHyphenIndex + 1, SecondHyphenIndex - FirstHyphenIndex - 1);
 				if (HotReloadedModuleName == ModuleName)
 				{
-					InOutModuleFilename = InOutModuleFilename.Mid(0, SecondHyphenIndex);
+					InOutModuleFilename.MidInline(0, SecondHyphenIndex, false);
 					SecondHyphenIndex = INDEX_NONE;
 				}
 			}
diff --git a/Engine/Source/Developer/LauncherServices/Private/Launcher/LauncherWorker.cpp b/Engine/Source/Developer/LauncherServices/Private/Launcher/LauncherWorker.cpp
--- a/Engine/Source/Developer/LauncherServices/Private/Launcher/LauncherWorker.cpp
+++ b/Engine/Source/Developer/LauncherServices/Private/Launcher/LauncherWorker.cpp
@@ -1087,10 +1087,10 @@ bool FLauncherWorker::TerminateLaunchedProcess()
 				FString TargetDeviceId = DeviceId;
 				
 				// remove the variant prefix (eg. Android_ETC@deviceId)
-				int32 InPos = TargetDeviceId.Find("@");
+				int32 InPos = TargetDeviceId.Find("@", ESearchCase::CaseSensitive);
 				if (InPos > 0) 
 				{ 
-					TargetDeviceId = TargetDeviceId.Right(TargetDeviceId.Len() -  InPos - 1);
+					TargetDeviceId.RightInline(TargetDeviceId.Len() -  InPos - 1, false);
 
 				}
 
diff --git a/Engine/Source/Developer/Merge/Private/SMergeGraphView.cpp b/Engine/Source/Developer/Merge/Private/SMergeGraphView.cpp
--- a/Engine/Source/Developer/Merge/Private/SMergeGraphView.cpp
+++ b/Engine/Source/Developer/Merge/Private/SMergeGraphView.cpp
@@ -462,7 +462,7 @@ void SMergeGraphView::Construct(const FArguments InArgs
 			int32 PeriodIndex = INDEX_NONE;
 			if (DisplayString.FindLastChar('.', PeriodIndex))
 			{
-				DisplayString = DisplayString.Mid(PeriodIndex + 1);
+				DisplayString.MidInline(PeriodIndex + 1, MAX_int32, false);
 			}
 
 			return SNew(SHorizontalBox)
diff --git a/Engine/Source/Developer/MessageLog/Private/UserInterface/SMessageLogMessageListRow.cpp b/Engine/Source/Developer/MessageLog/Private/UserInterface/SMessageLogMessageListRow.cpp
--- a/Engine/Source/Developer/MessageLog/Private/UserInterface/SMessageLogMessageListRow.cpp
+++ b/Engine/Source/Developer/MessageLog/Private/UserInterface/SMessageLogMessageListRow.cpp
@@ -283,7 +283,7 @@ void SMessageLogMessageListRow::CreateMessage(const TSharedRef<SHorizontalBox>&
 				int32 LineNumber = FCString::Atoi(*FileAndLineRegexMatcher.GetCaptureGroup(2));
 
 				// Remove the hyperlink from the message, since we're splitting it into its own string.
-				MessageString = MessageString.RightChop(FileAndLineRegexMatcher.GetMatchEnding());
+				MessageString.RightChopInline(FileAndLineRegexMatcher.GetMatchEnding(), false);
 
 				SourceLink = SNew(SHyperlink)
 					.Style(FEditorStyle::Get(), "Common.GotoNativeCodeHyperlink")
diff --git a/Engine/Source/Developer/PakFileUtilities/Private/PakFileUtilities.cpp b/Engine/Source/Developer/PakFileUtilities/Private/PakFileUtilities.cpp
--- a/Engine/Source/Developer/PakFileUtilities/Private/PakFileUtilities.cpp
+++ b/Engine/Source/Developer/PakFileUtilities/Private/PakFileUtilities.cpp
@@ -319,7 +319,7 @@ public:
 				if (Lines[EntryIndex].FindLastChar('"', OpenOrderNumber))
 				{
 					FString ReadNum = Lines[EntryIndex].RightChop(OpenOrderNumber + 1);
-					Lines[EntryIndex] = Lines[EntryIndex].Left(OpenOrderNumber + 1);
+					Lines[EntryIndex].LeftInline(OpenOrderNumber + 1, false);
 					ReadNum.TrimStartInline();
 					if (ReadNum.IsNumeric())
 					{
@@ -726,7 +726,7 @@ FString GetCommonRootPath(TArray<FPakInputPair>& FilesToAdd)
 		}
 		if ((CommonSeparatorIndex + 1) < Root.Len())
 		{
-			Root = Root.Mid(0, CommonSeparatorIndex + 1);
+			Root.MidInline(0, CommonSeparatorIndex + 1, false);
 		}
 	}
 	return Root;
@@ -1744,7 +1744,7 @@ void LoadKeyChain(const TCHAR* CmdLine, FKeyChain& OutCryptoSettings)
 			if (EncryptionKeyString.Len() > RequiredKeyLength)
 			{
 				UE_LOG(LogPakFile, Warning, TEXT("AES encryption key is more than %d characters long, so will be truncated!"), RequiredKeyLength);
-				EncryptionKeyString = EncryptionKeyString.Left(RequiredKeyLength);
+				EncryptionKeyString.LeftInline(RequiredKeyLength);
 			}
 
 			if (!FCString::IsPureAnsi(*EncryptionKeyString))
@@ -2610,7 +2610,7 @@ int32 GetPakPriorityFromFilename( const FString& PakFilename )
 		PakIndexFromFilename.FindLastChar('_', PakIndexStart);
 		if (PakIndexStart != INDEX_NONE)
 		{
-			PakIndexFromFilename = PakIndexFromFilename.RightChop(PakIndexStart + 1);
+			PakIndexFromFilename.RightChopInline(PakIndexStart + 1, false);
 			if (PakIndexFromFilename.IsNumeric())
 			{
 				PakPriority = FCString::Atoi(*PakIndexFromFilename);
diff --git a/Engine/Source/Developer/ProfilerClient/Private/ProfilerClientManager.cpp b/Engine/Source/Developer/ProfilerClient/Private/ProfilerClientManager.cpp
--- a/Engine/Source/Developer/ProfilerClient/Private/ProfilerClientManager.cpp
+++ b/Engine/Source/Developer/ProfilerClient/Private/ProfilerClientManager.cpp
@@ -908,7 +908,7 @@ int32 FServiceConnection::FindOrAddStat(const FStatNameAndInfo& StatNameAndInfo,
 		StatDescription.Name = !Description.IsEmpty() ? Description : StatName.ToString();
 		if(StatDescription.Name.Contains(TEXT("STAT_")))
 		{
-			StatDescription.Name = StatDescription.Name.RightChop(FString(TEXT("STAT_")).Len());
+			StatDescription.Name.RightChopInline(FString(TEXT("STAT_")).Len(), false);
 		}
 		StatDescription.StatType = StatType;
 
diff --git a/Engine/Source/Developer/ScriptDisassembler/Public/ScriptDisassembler.h b/Engine/Source/Developer/ScriptDisassembler/Public/ScriptDisassembler.h
--- a/Engine/Source/Developer/ScriptDisassembler/Public/ScriptDisassembler.h
+++ b/Engine/Source/Developer/ScriptDisassembler/Public/ScriptDisassembler.h
@@ -75,7 +75,7 @@ private:
 	void DropIndent()
 	{
 		// Blah, this is awful
-		Indents = Indents.Left(Indents.Len() - 2);
+		Indents.LeftInline(Indents.Len() - 2);
 	}
 
 	template <typename T>
diff --git a/Engine/Source/Developer/TraceInsights/Private/Insights/ViewModels/LogMessage.cpp b/Engine/Source/Developer/TraceInsights/Private/Insights/ViewModels/LogMessage.cpp
--- a/Engine/Source/Developer/TraceInsights/Private/Insights/ViewModels/LogMessage.cpp
+++ b/Engine/Source/Developer/TraceInsights/Private/Insights/ViewModels/LogMessage.cpp
@@ -40,7 +40,7 @@ FLogMessageRecord::FLogMessageRecord(const Trace::FLogMessage& TraceLogMessage)
 	FString CategoryStr(TraceLogMessage.Category->Name);
 	if (CategoryStr.StartsWith(TEXT("Log")))
 	{
-		CategoryStr = CategoryStr.RightChop(3);
+		CategoryStr.RightChopInline(3, false);
 	}
 	Category = FText::FromString(CategoryStr);
 
diff --git a/Engine/Source/Developer/TraceInsights/Private/Insights/ViewModels/MarkersTimingTrack.cpp b/Engine/Source/Developer/TraceInsights/Private/Insights/ViewModels/MarkersTimingTrack.cpp
--- a/Engine/Source/Developer/TraceInsights/Private/Insights/ViewModels/MarkersTimingTrack.cpp
+++ b/Engine/Source/Developer/TraceInsights/Private/Insights/ViewModels/MarkersTimingTrack.cpp
@@ -351,7 +351,7 @@ FLinearColor FTimeMarkerTrackBuilder::GetColorByCategory(const TCHAR* const Cate
 	FString CategoryStr(Category);
 	if (CategoryStr.StartsWith(TEXT("Log")))
 	{
-		CategoryStr = CategoryStr.RightChop(3);
+		CategoryStr.RightChopInline(3, false);
 	}
 
 	uint32 Hash = 0;
@@ -424,7 +424,7 @@ void FTimeMarkerTrackBuilder::Flush(float AvailableTextW)
 			FString CategoryStr(LastCategory);
 			if (CategoryStr.StartsWith(TEXT("Log")))
 			{
-				CategoryStr = CategoryStr.RightChop(3);
+				CategoryStr.RightChopInline(3, fase);
 			}
 
 			const int32 LastWholeCharacterIndexCategory = FontMeasureService->FindLastWholeCharacterIndexBeforeOffset(CategoryStr, Font, FMath::RoundToInt(AvailableTextW - 2.0f));
diff --git a/Engine/Source/Developer/TraceInsights/Private/Insights/Widgets/SLogView.cpp b/Engine/Source/Developer/TraceInsights/Private/Insights/Widgets/SLogView.cpp
--- a/Engine/Source/Developer/TraceInsights/Private/Insights/Widgets/SLogView.cpp
+++ b/Engine/Source/Developer/TraceInsights/Private/Insights/Widgets/SLogView.cpp
@@ -741,7 +741,7 @@ void SLogView::Tick(const FGeometry& AllottedGeometry, const double InCurrentTim
 				FString CategoryStr(Category.Name);
 				if (CategoryStr.StartsWith(TEXT("Log")))
 				{
-					CategoryStr = CategoryStr.RightChop(3);
+					CategoryStr.RightChopInline(3, false);
 				}
 				if (Categories.Contains(FName(*CategoryStr)))
 				{
diff --git a/Engine/Source/Developer/TraceServices/Private/SessionService.cpp b/Engine/Source/Developer/TraceServices/Private/SessionService.cpp
--- a/Engine/Source/Developer/TraceServices/Private/SessionService.cpp
+++ b/Engine/Source/Developer/TraceServices/Private/SessionService.cpp
@@ -246,7 +246,7 @@ bool FSessionService::ConnectSession(const TCHAR* ControlClientAddress)
 	{
 		FString PortString = AddressString.RightChop(LastColonIndex + 1);
 		Port = FCString::Atoi(*PortString);
-		AddressString = AddressString.Left(LastColonIndex);
+		AddressString.LeftInline(LastColonIndex);
 	}
 	
 	TSharedPtr<FInternetAddr> ClientAddr = Sockets->GetAddressFromString(AddressString);
diff --git a/Engine/Source/Developer/TreeMap/TreeMap.cpp b/Engine/Source/Developer/TreeMap/TreeMap.cpp
--- a/Engine/Source/Developer/TreeMap/TreeMap.cpp
+++ b/Engine/Source/Developer/TreeMap/TreeMap.cpp
@@ -613,12 +613,12 @@ FTreeMapNodeDataPtr ITreeMap::ParseOPMLToTreeMapData( const FString& OPMLFilePat
 
 												if( HashTagLength > 1 )
 												{
-													const auto HashTag = ChildNodeData->Name.Mid( HashTagCharIndex + 1, HashTagLength - 1 );
+													const FString HashTag = ChildNodeData->Name.Mid( HashTagCharIndex + 1, HashTagLength - 1 );
 
 													ChildNodeData->HashTags.Add( HashTag );
 
 													// Strip the hash tag ofg of the original string
-													ChildNodeData->Name = ChildNodeData->Name.Mid( 0, HashTagCharIndex );
+													ChildNodeData->Name.MidInline( 0, HashTagCharIndex, false );
 													if( HashTagCharIndex + HashTagLength < ChildNodeData->Name.Len() )
 													{
 														ChildNodeData->Name += ChildNodeData->Name.Mid( HashTagCharIndex + HashTagLength );
diff --git a/Engine/Source/Editor/AIGraph/Private/AIGraphTypes.cpp b/Engine/Source/Editor/AIGraph/Private/AIGraphTypes.cpp
--- a/Engine/Source/Editor/AIGraph/Private/AIGraphTypes.cpp
+++ b/Engine/Source/Editor/AIGraph/Private/AIGraphTypes.cpp
@@ -62,10 +62,10 @@ FString FGraphNodeClassData::ToString() const
 			return ClassDesc.LeftChop(2);
 		}
 
-		const int32 ShortNameIdx = ClassDesc.Find(TEXT("_"));
+		const int32 ShortNameIdx = ClassDesc.Find(TEXT("_"), ESearchCase::CaseSensitive);
 		if (ShortNameIdx != INDEX_NONE)
 		{
-			ClassDesc = ClassDesc.Mid(ShortNameIdx + 1);
+			ClassDesc.MidInline(ShortNameIdx + 1, MAX_int32, false);
 		}
 
 		return ClassDesc;
diff --git a/Engine/Source/Editor/AddContentDialog/Private/ContentSourceProviders/FeaturePack/FeaturePackContentSource.cpp b/Engine/Source/Editor/AddContentDialog/Private/ContentSourceProviders/FeaturePack/FeaturePackContentSource.cpp
--- a/Engine/Source/Editor/AddContentDialog/Private/ContentSourceProviders/FeaturePack/FeaturePackContentSource.cpp
+++ b/Engine/Source/Editor/AddContentDialog/Private/ContentSourceProviders/FeaturePack/FeaturePackContentSource.cpp
@@ -188,7 +188,7 @@ FFeaturePackContentSource::FFeaturePackContentSource(FString InFeaturePackPath)
 		if( ThisTemplateRoot.EndsWith(TemplatesFolder) == true )
 		{
 			int32 Index = ThisTemplateRoot.Find(TemplatesFolder);
-			ThisTemplateRoot = ThisTemplateRoot.Left(Index);
+			ThisTemplateRoot.LeftInline(Index);
 		}
 		MountPoint = ThisTemplateRoot;
 		FFileHelper::LoadFileToString( ManifestString, *FeaturePackPath);
diff --git a/Engine/Source/Editor/BehaviorTreeEditor/Private/BehaviorTreeGraphNode_CompositeDecorator.cpp b/Engine/Source/Editor/BehaviorTreeEditor/Private/BehaviorTreeGraphNode_CompositeDecorator.cpp
--- a/Engine/Source/Editor/BehaviorTreeEditor/Private/BehaviorTreeGraphNode_CompositeDecorator.cpp
+++ b/Engine/Source/Editor/BehaviorTreeEditor/Private/BehaviorTreeGraphNode_CompositeDecorator.cpp
@@ -296,7 +296,7 @@ void UpdateLogicOpStack(TArray<FLogicDesc>& OpStack, FString& Description, FStri
 		if (OpStack[LastIdx].NumLeft <= 0)
 		{
 			OpStack.RemoveAt(LastIdx);
-			Indent = Indent.LeftChop(2);
+			Indent.LeftChopInline(2, false);
 
 			UpdateLogicOpStack(OpStack, Description, Indent);
 		}
diff --git a/Engine/Source/Editor/BlueprintGraph/Private/K2Node.cpp b/Engine/Source/Editor/BlueprintGraph/Private/K2Node.cpp
--- a/Engine/Source/Editor/BlueprintGraph/Private/K2Node.cpp
+++ b/Engine/Source/Editor/BlueprintGraph/Private/K2Node.cpp
@@ -1829,7 +1829,7 @@ void UK2Node::GetPinHoverText(const UEdGraphPin& Pin, FString& HoverTextOut) con
 			if (LineCounter >= MaxArrayPinTooltipLineCount)
 			{
 				// truncate WatchText so it contains a finite number of lines
-				WatchText  = WatchText.Left(NewWatchTextLen);
+				WatchText.LeftInline(NewWatchTextLen, false);
 				WatchText += "..."; // WatchText should already have a trailing newline (no need to prepend this with one)
 				break;
 			}
diff --git a/Engine/Source/Editor/ContentBrowser/Private/AssetViewWidgets.cpp b/Engine/Source/Editor/ContentBrowser/Private/AssetViewWidgets.cpp
--- a/Engine/Source/Editor/ContentBrowser/Private/AssetViewWidgets.cpp
+++ b/Engine/Source/Editor/ContentBrowser/Private/AssetViewWidgets.cpp
@@ -1426,7 +1426,8 @@ void SAssetViewItem::CacheDisplayTags()
 				{
 					// Remove the class path for native classes, and also remove Engine. for engine classes
 					const int32 SizeOfPrefix = UE_ARRAY_COUNT(StringToRemove) - 1;
-					ValueString = ValueString.Mid(SizeOfPrefix, ValueString.Len() - SizeOfPrefix).Replace(TEXT("Engine."), TEXT(""));
+					ValueString.MidInline(SizeOfPrefix, ValueString.Len() - SizeOfPrefix, false);
+					ValueString.ReplaceInline(TEXT("Engine."), TEXT(""));
 				}
 
 				if (TagField)
@@ -1452,7 +1453,7 @@ void SAssetViewItem::CacheDisplayTags()
 							const FString EnumPrefix = TagEnum->GenerateEnumPrefix();
 							if (EnumPrefix.Len() && ValueString.StartsWith(EnumPrefix))
 							{
-								ValueString = ValueString.RightChop(EnumPrefix.Len() + 1);	// +1 to skip over the underscore
+								ValueString.RightChopInline(EnumPrefix.Len() + 1, false);	// +1 to skip over the underscore
 							}
 						}
 
diff --git a/Engine/Source/Editor/ContentBrowser/Private/ContentBrowserUtils.cpp b/Engine/Source/Editor/ContentBrowser/Private/ContentBrowserUtils.cpp
--- a/Engine/Source/Editor/ContentBrowser/Private/ContentBrowserUtils.cpp
+++ b/Engine/Source/Editor/ContentBrowser/Private/ContentBrowserUtils.cpp
@@ -1282,13 +1282,13 @@ bool ContentBrowserUtils::IsClassesFolder(const FString& InPath)
 	// Strip off any leading or trailing forward slashes
 	// We just want the name without any path separators
 	FString CleanFolderPath = InPath;
-	while ( CleanFolderPath.StartsWith(TEXT("/")) )
+	while ( CleanFolderPath.StartsWith(TEXT("/"), ESearchCase::CaseSensitive) )
 	{
-		CleanFolderPath = CleanFolderPath.Mid(1);
+		CleanFolderPath.MidInline(1, MAX_int32, false);
 	}
-	while ( CleanFolderPath.EndsWith(TEXT("/")) )
+	while ( CleanFolderPath.EndsWith(TEXT("/"), ESearchCase::CaseSensitive) )
 	{
-		CleanFolderPath = CleanFolderPath.Mid(0, CleanFolderPath.Len() - 1);
+		CleanFolderPath.MidInline(0, CleanFolderPath.Len() - 1, false);
 	}
 
 	static const FString ClassesPrefix = TEXT("Classes_");
@@ -1414,13 +1414,13 @@ FText ContentBrowserUtils::GetRootDirDisplayName(const FString& FolderPath)
 	// Strip off any leading or trailing forward slashes
 	// We just want the name without any path separators
 	FString CleanFolderPath = FolderPath;
-	while(CleanFolderPath.StartsWith(TEXT("/")))
+	while(CleanFolderPath.StartsWith(TEXT("/"), ESearchCase::CaseSensitive))
 	{
-		CleanFolderPath = CleanFolderPath.Mid(1);
+		CleanFolderPath.MidInline(1, MAX_int32, false);
 	}
-	while(CleanFolderPath.EndsWith(TEXT("/")))
+	while(CleanFolderPath.EndsWith(TEXT("/"), ESearchCase::CaseSensitive))
 	{
-		CleanFolderPath = CleanFolderPath.Mid(0, CleanFolderPath.Len() - 1);
+		CleanFolderPath.MidInline(0, CleanFolderPath.Len() - 1, false);
 	}
 
 	static const FString ClassesPrefix = TEXT("Classes_");
@@ -1429,7 +1429,7 @@ FText ContentBrowserUtils::GetRootDirDisplayName(const FString& FolderPath)
 	// Strip off the "Classes_" prefix
 	if(bIsClassDir)
 	{
-		CleanFolderPath = CleanFolderPath.Mid(ClassesPrefix.Len());
+		CleanFolderPath.MidInline(ClassesPrefix.Len(), MAX_int32, false);
 	}
 
 	// Also localize well known folder names, like "Engine" and "Game"
@@ -2175,7 +2175,7 @@ void ContentBrowserUtils::SyncPathsFromSourceControl(const TArray<FString>& Cont
 				if (PackagePath.Len() > 1 && PackagePath[PackagePath.Len() - 1] == TEXT('/'))
 				{
 					// The filter path can't end with a trailing slash
-					PackagePath = PackagePath.LeftChop(1);
+					PackagePath.LeftChopInline(1, false);
 				}
 				Filter.PackagePaths.Emplace(*PackagePath);
 			}
diff --git a/Engine/Source/Editor/ContentBrowser/Private/EmptyFolderVisibilityManager.cpp b/Engine/Source/Editor/ContentBrowser/Private/EmptyFolderVisibilityManager.cpp
--- a/Engine/Source/Editor/ContentBrowser/Private/EmptyFolderVisibilityManager.cpp
+++ b/Engine/Source/Editor/ContentBrowser/Private/EmptyFolderVisibilityManager.cpp
@@ -61,7 +61,7 @@ void FEmptyFolderVisibilityManager::SetAlwaysShowPath(const FString& InPath)
 	bool bWasAlreadyShown = false;
 	if (PathToAdd[PathToAdd.Len() - 1] == TEXT('/'))
 	{
-		PathToAdd = PathToAdd.Mid(0, PathToAdd.Len() - 1);
+		PathToAdd.MidInline(0, PathToAdd.Len() - 1, false);
 	}
 	PathsToAlwaysShow.Add(PathToAdd, &bWasAlreadyShown);
 
diff --git a/Engine/Source/Editor/ContentBrowser/Private/SPathView.cpp b/Engine/Source/Editor/ContentBrowser/Private/SPathView.cpp
--- a/Engine/Source/Editor/ContentBrowser/Private/SPathView.cpp
+++ b/Engine/Source/Editor/ContentBrowser/Private/SPathView.cpp
@@ -1073,13 +1073,13 @@ void SPathView::Populate()
 			// Strip off any leading or trailing forward slashes.  We just want a root path name that
 			// we can display, and we'll add the path separators back later on
 			FString CleanRootPathName = *RootPathIt;
-			while( CleanRootPathName.StartsWith( TEXT( "/" ) ) )
+			while( CleanRootPathName.StartsWith( TEXT( "/" ), ESearchCase::CaseSensitive ) )
 			{
-				CleanRootPathName = CleanRootPathName.Mid( 1 );
+				CleanRootPathName.MidInline( 1, MAX_int32, false);
 			}
-			while( CleanRootPathName.EndsWith( TEXT( "/" ) ) )
+			while( CleanRootPathName.EndsWith( TEXT( "/" ), ESearchCase::CaseSensitive) )
 			{
-				CleanRootPathName = CleanRootPathName.Mid( 0, CleanRootPathName.Len() - 1 );
+				CleanRootPathName.MidInline( 0, CleanRootPathName.Len() - 1, false );
 			}
 
 			// Templates can mount "root" items which are actually sub-items under a root (see FUnrealEdMisc::MountTemplateSharedPaths)
@@ -1167,14 +1167,14 @@ void SPathView::SortRootItems()
 		const bool bOneIsClass = OneModuleName.StartsWith(ClassesPrefix);
 		if(bOneIsClass)
 		{
-			OneModuleName = OneModuleName.Mid(ClassesPrefix.Len());
+			OneModuleName.MidInline(ClassesPrefix.Len(), MAX_int32, false);
 		}
 
 		FString TwoModuleName = Two->FolderName;
 		const bool bTwoIsClass = TwoModuleName.StartsWith(ClassesPrefix);
 		if(bTwoIsClass)
 		{
-			TwoModuleName = TwoModuleName.Mid(ClassesPrefix.Len());
+			TwoModuleName.MidInline(ClassesPrefix.Len(), MAX_int32, false);
 		}
 
 		// We want to sort content before classes if both items belong to the same module
diff --git a/Engine/Source/Editor/DetailCustomizations/Private/CameraDetails.cpp b/Engine/Source/Editor/DetailCustomizations/Private/CameraDetails.cpp
--- a/Engine/Source/Editor/DetailCustomizations/Private/CameraDetails.cpp
+++ b/Engine/Source/Editor/DetailCustomizations/Private/CameraDetails.cpp
@@ -256,7 +256,7 @@ void FCameraDetails::OnCommitAspectRatioText(const FText& ItemFText, ETextCommit
 			FString RemainingText = ItemText.Mid(DelimIdx + 1).TrimStart();
 			if (RemainingText.FindChar(TCHAR(' '), WSIdx))
 			{
-				RemainingText = RemainingText.Left(WSIdx);
+				RemainingText.LeftInline(WSIdx, false);
 			}
 			int32 Height;
 			TTypeFromString<int32>::FromString(Height, *RemainingText);
diff --git a/Engine/Source/Editor/DetailCustomizations/Private/DirectoryPathStructCustomization.cpp b/Engine/Source/Editor/DetailCustomizations/Private/DirectoryPathStructCustomization.cpp
--- a/Engine/Source/Editor/DetailCustomizations/Private/DirectoryPathStructCustomization.cpp
+++ b/Engine/Source/Editor/DetailCustomizations/Private/DirectoryPathStructCustomization.cpp
@@ -154,7 +154,7 @@ FReply FDirectoryPathStructCustomization::OnPickDirectory(TSharedRef<IPropertyHa
 
 					if (bRelativeToGameContentDir)
 					{
-						Directory = Directory.RightChop(AbsoluteGameContentDir.Len());
+						Directory.RightChopInline(AbsoluteGameContentDir.Len(), false);
 					}
 					else if (bUseRelativePath)
 					{
diff --git a/Engine/Source/Editor/Documentation/Private/DocumentationLink.h b/Engine/Source/Editor/Documentation/Private/DocumentationLink.h
--- a/Engine/Source/Editor/Documentation/Private/DocumentationLink.h
+++ b/Engine/Source/Editor/Documentation/Private/DocumentationLink.h
@@ -212,17 +212,17 @@ private:
 				Anchor = FString( TEXT("#") ) + Anchor;
 			}
 
-			if ( Anchor.EndsWith( TEXT("/") ) )
+			if ( Anchor.EndsWith( TEXT("/"), ESearchCase::CaseSensitive ) )
 			{
-				Anchor = Anchor.Left( Anchor.Len() - 1 );
+				Anchor.LeftInline( Anchor.Len() - 1 );
 			}
 
-			if ( PathAndQueryString.EndsWith( TEXT("/") ) )
+			if ( PathAndQueryString.EndsWith( TEXT("/"), ESearchCase::CaseSensitive ) )
 			{
-				PathAndQueryString = PathAndQueryString.Left(PathAndQueryString.Len() - 1);
+				PathAndQueryString.LeftInline(PathAndQueryString.Len() - 1, false);
 			}
 
-			if ( !PathAndQueryString.IsEmpty() && !PathAndQueryString.StartsWith( TEXT("/") ) )
+			if ( !PathAndQueryString.IsEmpty() && !PathAndQueryString.StartsWith( TEXT("/"), ESearchCase::CaseSensitive ) )
 			{
 				PathAndQueryString = FString(TEXT("/")) + PathAndQueryString;
 			}
diff --git a/Engine/Source/Editor/EnvironmentQueryEditor/Private/SEnvQueryProfiler.cpp b/Engine/Source/Editor/EnvironmentQueryEditor/Private/SEnvQueryProfiler.cpp
--- a/Engine/Source/Editor/EnvironmentQueryEditor/Private/SEnvQueryProfiler.cpp
+++ b/Engine/Source/Editor/EnvironmentQueryEditor/Private/SEnvQueryProfiler.cpp
@@ -201,7 +201,7 @@ static UObject* FindQueryObjectByName(FName StatName)
 	int32 SepIdx = INDEX_NONE;
 	if (AssetPathName.FindLastChar(TEXT('_'), SepIdx))
 	{
-		AssetPathName = AssetPathName.Left(SepIdx);
+		AssetPathName.LeftInline(SepIdx);
 	}
 
 	UObject* QueryOb = FindObject<UObject>(ANY_PACKAGE, *AssetPathName);
diff --git a/Engine/Source/Editor/GameProjectGeneration/Private/GameProjectUtils.cpp b/Engine/Source/Editor/GameProjectGeneration/Private/GameProjectUtils.cpp
--- a/Engine/Source/Editor/GameProjectGeneration/Private/GameProjectUtils.cpp
+++ b/Engine/Source/Editor/GameProjectGeneration/Private/GameProjectUtils.cpp
@@ -271,11 +271,11 @@ FText FNewClassInfo::GetClassDescription(const bool bFullDescription/* = true*/)
 					if(ClassDescription.FindChar('.', FullStopIndex))
 					{
 						// Only show the first sentence so as not to clutter up the UI with a detailed description of implementation details
-						ClassDescription = ClassDescription.Left(FullStopIndex + 1);
+						ClassDescription.LeftInline(FullStopIndex + 1, false);
 					}
 
 					// Strip out any new-lines in the description
-					ClassDescription.ReplaceInline(TEXT("\n"), TEXT(" "));
+					ClassDescription.ReplaceInline(TEXT("\n"), TEXT(" "), ESearchCase::CaseSensitive);
 				}
 
 				return FText::FromString(ClassDescription);
@@ -370,11 +370,11 @@ FString FNewClassInfo::GetCleanClassName(const FString& ClassName) const
 			// if our class ends with either Widget or WidgetStyle, we need to strip those out to avoid silly looking duplicates
 			if(CleanClassName.EndsWith(TEXT("Style")))
 			{
-				CleanClassName = CleanClassName.LeftChop(5); // 5 for "Style"
+				CleanClassName.LeftChopInline(5, false); // 5 for "Style"
 			}
 			if(CleanClassName.EndsWith(TEXT("Widget")))
 			{
-				CleanClassName = CleanClassName.LeftChop(6); // 6 for "Widget"
+				CleanClassName.LeftChopInline(6, false); // 6 for "Widget"
 			}
 		}
 		break;
@@ -2754,7 +2754,8 @@ GameProjectUtils::EProjectDuplicateResult GameProjectUtils::DuplicateProjectForU
 			break;
 		}
 
-		NewDirectoryName = NewDirectoryName.Left(LastSpace).TrimEnd();
+		NewDirectoryName.LeftInline(LastSpace, false);
+		NewDirectoryName.TrimEndInline();
 	}
 
 	// Append the new version number
diff --git a/Engine/Source/Editor/GraphEditor/Private/KismetPins/SGraphPinCollisionProfile.cpp b/Engine/Source/Editor/GraphEditor/Private/KismetPins/SGraphPinCollisionProfile.cpp
--- a/Engine/Source/Editor/GraphEditor/Private/KismetPins/SGraphPinCollisionProfile.cpp
+++ b/Engine/Source/Editor/GraphEditor/Private/KismetPins/SGraphPinCollisionProfile.cpp
@@ -103,16 +103,16 @@ void SGraphPinCollisionProfile::GetPropertyAsName(FName& OutName) const
 
 	FString PinString = GraphPinObj->GetDefaultAsString();
 
-	if (PinString.StartsWith(TEXT("(")) && PinString.EndsWith(TEXT(")")))
+	if (PinString.StartsWith(TEXT("("), ESearchCase::CaseSensitive) && PinString.EndsWith(TEXT(")"), ESearchCase::CaseSensitive))
 	{
-		PinString = PinString.LeftChop(1);
-		PinString = PinString.RightChop(1);
-		PinString.Split("=", NULL, &PinString);
+		PinString.LeftChopInline(1, false);
+		PinString.RightChopInline(1, false);
+		PinString.Split(TEXT("="), nullptr, &PinString, ESearchCase::CaseSensitive);
 
-		if (PinString.StartsWith(TEXT("\"")) && PinString.EndsWith(TEXT("\"")))
+		if (PinString.StartsWith(TEXT("\""), ESearchCase::CaseSensitive) && PinString.EndsWith(TEXT("\""), ESearchCase::CaseSensitive))
 		{
-			PinString = PinString.LeftChop(1);
-			PinString = PinString.RightChop(1);
+			PinString.LeftChopInline(1, false);
+			PinString.RightChopInline(1, false);
 		}
 	}
 
diff --git a/Engine/Source/Editor/IntroTutorials/Private/STutorialEditableText.cpp b/Engine/Source/Editor/IntroTutorials/Private/STutorialEditableText.cpp
--- a/Engine/Source/Editor/IntroTutorials/Private/STutorialEditableText.cpp
+++ b/Engine/Source/Editor/IntroTutorials/Private/STutorialEditableText.cpp
@@ -439,7 +439,7 @@ void STutorialEditableText::HandleHyperlinkComboOpened()
 		{
 			if(FChar::IsLinebreak(SelectedText[SelectedTextIndex]))
 			{
-				SelectedText = SelectedText.Left(SelectedTextIndex);
+				SelectedText.LeftInline(SelectedTextIndex, false);
 				break;
 			}
 		}
diff --git a/Engine/Source/Editor/LocalizationDashboard/Private/GatherTextDetailCustomizations.cpp b/Engine/Source/Editor/LocalizationDashboard/Private/GatherTextDetailCustomizations.cpp
--- a/Engine/Source/Editor/LocalizationDashboard/Private/GatherTextDetailCustomizations.cpp
+++ b/Engine/Source/Editor/LocalizationDashboard/Private/GatherTextDetailCustomizations.cpp
@@ -185,7 +185,7 @@ namespace
 			CurrentPath.FindLastChar(TEXT('*'), WildcardStarIndex);
 			if (WildcardStarIndex != INDEX_NONE)
 			{
-				CurrentPath = CurrentPath.Left(WildcardStarIndex);
+				CurrentPath.LeftInline(WildcardStarIndex, false);
 			}
 		}
 
diff --git a/Engine/Source/Editor/MaterialEditor/Private/MaterialExpressionClasses.cpp b/Engine/Source/Editor/MaterialEditor/Private/MaterialExpressionClasses.cpp
--- a/Engine/Source/Editor/MaterialEditor/Private/MaterialExpressionClasses.cpp
+++ b/Engine/Source/Editor/MaterialEditor/Private/MaterialExpressionClasses.cpp
@@ -126,7 +126,7 @@ void MaterialExpressionClasses::InitMaterialExpressionClasses()
 
 							if (ClassName.StartsWith(ExpressionPrefix, ESearchCase::CaseSensitive))
 							{
-								ClassName = ClassName.Mid(ExpressionPrefix.Len());
+								ClassName.MidInline(ExpressionPrefix.Len(), MAX_int32, false);
 							}
 							MaterialExpression.Name = ClassName;
 							MaterialExpression.MaterialClass = Class;
diff --git a/Engine/Source/Editor/MaterialEditor/Private/MaterialStatsCommon.cpp b/Engine/Source/Editor/MaterialEditor/Private/MaterialStatsCommon.cpp
--- a/Engine/Source/Editor/MaterialEditor/Private/MaterialStatsCommon.cpp
+++ b/Engine/Source/Editor/MaterialEditor/Private/MaterialStatsCommon.cpp
@@ -176,7 +176,7 @@ FString FMaterialStatsUtils::ShaderPlatformTypeName(const EShaderPlatform Platfo
 			FString FormatName = ShaderPlatformToShaderFormatName(PlatformID).ToString();
 			if (FormatName.StartsWith(TEXT("SF_")))
 			{
-				FormatName = FormatName.Mid(3);
+				FormatName.MidInline(3, MAX_int32, false);
 			}
 			return FormatName;
 	}
diff --git a/Engine/Source/Editor/Matinee/Private/MatineeTrackHelpers.cpp b/Engine/Source/Editor/Matinee/Private/MatineeTrackHelpers.cpp
--- a/Engine/Source/Editor/Matinee/Private/MatineeTrackHelpers.cpp
+++ b/Engine/Source/Editor/Matinee/Private/MatineeTrackHelpers.cpp
@@ -638,14 +638,14 @@ void  UMatineeTrackFloatPropHelper::PostCreateTrack( UInterpTrack *Track, bool b
 
 		// Set track title to property name (cut off component name if there is one).
 		FString PropString = TrackAddPropName.ToString();
-		int32 PeriodPos = PropString.Find(TEXT("."));
+		int32 PeriodPos = PropString.Find(TEXT("."), ESearchCase::CaseSensitive);
 		if(PeriodPos != INDEX_NONE)
 		{
-			PropString = PropString.Mid(PeriodPos+1);
+			PropString.MidInline(PeriodPos+1, MAX_int32, false);
 		}
 
 		PropTrack->PropertyName = TrackAddPropName;
-		PropTrack->TrackTitle = *PropString;
+		PropTrack->TrackTitle = MoveTemp(PropString);
 
 		TrackAddPropName = NAME_None;
 	}
@@ -777,14 +777,14 @@ void UMatineeTrackBoolPropHelper::PostCreateTrack( UInterpTrack* Track, bool bDu
 
 		// Set track title to property name (cut off component name if there is one).
 		FString PropString = TrackAddPropName.ToString();
-		int32 PeriodPos = PropString.Find(TEXT("."));
+		int32 PeriodPos = PropString.Find(TEXT("."), ESearchCase::CaseSensitive);
 		if(PeriodPos != INDEX_NONE)
 		{
-			PropString = PropString.Mid(PeriodPos+1);
+			PropString.MidInline(PeriodPos+1, MAX_int32, false);
 		}
 
 		PropTrack->PropertyName = TrackAddPropName;
-		PropTrack->TrackTitle = *PropString;
+		PropTrack->TrackTitle = MoveTemp(PropString);
 
 		TrackAddPropName = NAME_None;
 	}
@@ -1003,14 +1003,14 @@ void  UMatineeTrackVectorPropHelper::PostCreateTrack( UInterpTrack *Track, bool
 
 		// Set track title to property name (cut off component name if there is one).
 		FString PropString = TrackAddPropName.ToString();
-		int32 PeriodPos = PropString.Find(TEXT("."));
+		int32 PeriodPos = PropString.Find(TEXT("."), ESearchCase::CaseSensitive);
 		if(PeriodPos != INDEX_NONE)
 		{
-			PropString = PropString.Mid(PeriodPos+1);
+			PropString.MidInline(PeriodPos+1, MAX_int32, false);
 		}
 
 		PropTrack->PropertyName = TrackAddPropName;
-		PropTrack->TrackTitle = *PropString;
+		PropTrack->TrackTitle = MoveTemp(PropString);
 
 		TrackAddPropName = NAME_None;
 	}
@@ -1066,14 +1066,14 @@ void  UMatineeTrackColorPropHelper::PostCreateTrack( UInterpTrack *Track, bool b
 
 		// Set track title to property name (cut off component name if there is one).
 		FString PropString = TrackAddPropName.ToString();
-		int32 PeriodPos = PropString.Find(TEXT("."));
+		int32 PeriodPos = PropString.Find(TEXT("."), ESearchCase::CaseSensitive);
 		if(PeriodPos != INDEX_NONE)
 		{
-			PropString = PropString.Mid(PeriodPos+1);
+			PropString.MidInline(PeriodPos+1, MAX_int32, false);
 		}
 
 		PropTrack->PropertyName = TrackAddPropName;
-		PropTrack->TrackTitle = *PropString;
+		PropTrack->TrackTitle = MoveTemp(PropString);
 
 		TrackAddPropName = NAME_None;
 	}
@@ -1131,14 +1131,14 @@ void  UMatineeTrackLinearColorPropHelper::PostCreateTrack( UInterpTrack *Track,
 
 		// Set track title to property name (cut off component name if there is one).
 		FString PropString = TrackAddPropName.ToString();
-		int32 PeriodPos = PropString.Find(TEXT("."));
+		int32 PeriodPos = PropString.Find(TEXT("."), ESearchCase::CaseSensitive);
 		if(PeriodPos != INDEX_NONE)
 		{
-			PropString = PropString.Mid(PeriodPos+1);
+			PropString.MidInline(PeriodPos+1, MAX_int32, false);
 		}
 
 		PropTrack->PropertyName = TrackAddPropName;
-		PropTrack->TrackTitle = *PropString;
+		PropTrack->TrackTitle = MoveTemp(PropString);
 
 		TrackAddPropName = NAME_None;
 	}
diff --git a/Engine/Source/Editor/MovieSceneTools/Private/MovieSceneToolHelpers.cpp b/Engine/Source/Editor/MovieSceneTools/Private/MovieSceneToolHelpers.cpp
--- a/Engine/Source/Editor/MovieSceneTools/Private/MovieSceneToolHelpers.cpp
+++ b/Engine/Source/Editor/MovieSceneTools/Private/MovieSceneToolHelpers.cpp
@@ -437,7 +437,7 @@ bool MovieSceneToolHelpers::GetTakeNumber(const UMovieSceneSection* Section, FAs
 	FString ShotPackagePath = ShotData.PackagePath.ToString();
 	int32 ShotLastSlashPos = INDEX_NONE;
 	ShotPackagePath.FindLastChar(TCHAR('/'), ShotLastSlashPos);
-	ShotPackagePath = ShotPackagePath.Left(ShotLastSlashPos);
+	ShotPackagePath.LeftInline(ShotLastSlashPos, false);
 
 	FString ShotPrefix;
 	uint32 ShotNumber = INDEX_NONE;
@@ -464,7 +464,7 @@ bool MovieSceneToolHelpers::GetTakeNumber(const UMovieSceneSection* Section, FAs
 				FString AssetPackagePath = AssetObject.PackagePath.ToString();
 				int32 AssetLastSlashPos = INDEX_NONE;
 				AssetPackagePath.FindLastChar(TCHAR('/'), AssetLastSlashPos);
-				AssetPackagePath = AssetPackagePath.Left(AssetLastSlashPos);
+				AssetPackagePath.LeftInline(AssetLastSlashPos, false);
 
 				if (AssetPackagePath == ShotPackagePath)
 				{
diff --git a/Engine/Source/Editor/MovieSceneTools/Private/MovieSceneToolsProjectSettingsCustomization.cpp b/Engine/Source/Editor/MovieSceneTools/Private/MovieSceneToolsProjectSettingsCustomization.cpp
--- a/Engine/Source/Editor/MovieSceneTools/Private/MovieSceneToolsProjectSettingsCustomization.cpp
+++ b/Engine/Source/Editor/MovieSceneTools/Private/MovieSceneToolsProjectSettingsCustomization.cpp
@@ -37,6 +37,6 @@ void FMovieSceneToolsProjectSettingsCustomization::OnTakeSeparatorUpdated()
 	}
 	else if (ProjectSettings->TakeSeparator.Len() > 1)
 	{
-		ProjectSettings->TakeSeparator = ProjectSettings->TakeSeparator.LeftChop(ProjectSettings->TakeSeparator.Len()-1);
+		ProjectSettings->TakeSeparator.LeftChopInline(ProjectSettings->TakeSeparator.Len()-1);
 	}
 }
diff --git a/Engine/Source/Editor/MovieSceneTools/Private/MovieSceneTranslatorEDL.cpp b/Engine/Source/Editor/MovieSceneTools/Private/MovieSceneTranslatorEDL.cpp
--- a/Engine/Source/Editor/MovieSceneTools/Private/MovieSceneTranslatorEDL.cpp
+++ b/Engine/Source/Editor/MovieSceneTools/Private/MovieSceneTranslatorEDL.cpp
@@ -190,9 +190,9 @@ void ParseFromEDL(const FString& InputString, FFrameRate TickResolution, FFrameR
 				ReelName = InputChars[4];
 
 				//@todo can't assume avi's written out
-				ReelName = ReelName.LeftChop(4); // strip .avi
+				ReelName.LeftChopInline(4, false); // strip .avi
 
-				FString ElementName = ReelName;
+				FString ElementName = MoveTemp(ReelName);
 				FString ElementPath = ElementName;
 
 				// If everything checks out add to OutShotData
diff --git a/Engine/Source/Editor/PropertyEditor/Private/ItemPropertyNode.cpp b/Engine/Source/Editor/PropertyEditor/Private/ItemPropertyNode.cpp
--- a/Engine/Source/Editor/PropertyEditor/Private/ItemPropertyNode.cpp
+++ b/Engine/Source/Editor/PropertyEditor/Private/ItemPropertyNode.cpp
@@ -573,7 +573,7 @@ FText FItemPropertyNode::GetDisplayName() const
 
 					if (EndIndex != -1)
 					{
-						TypeName = TypeName.Mid(0, EndIndex);
+						TypeName.MidInline(0, EndIndex, false);
 					}
 				}
 
diff --git a/Engine/Source/Editor/PropertyEditor/Private/PropertyHandleImpl.cpp b/Engine/Source/Editor/PropertyEditor/Private/PropertyHandleImpl.cpp
--- a/Engine/Source/Editor/PropertyEditor/Private/PropertyHandleImpl.cpp
+++ b/Engine/Source/Editor/PropertyEditor/Private/PropertyHandleImpl.cpp
@@ -874,21 +874,21 @@ FPropertyAccess::Result FPropertyValueImpl::SetValueAsString( const FString& InV
 		{
 			while ( true )
 			{
-				if ( Value.StartsWith( TEXT("_") ) )
+				if ( Value.StartsWith( TEXT("_"), ESearchCase::CaseSensitive ) )
 				{
 					// Strip leading underscores.
 					do
 					{
-						Value = Value.Right( Value.Len()-1 );
-					} while ( Value.StartsWith( TEXT("_") ) );
+						Value.RightInline( Value.Len()-1, false);
+					} while ( Value.StartsWith( TEXT("_"), ESearchCase::CaseSensitive ) );
 				}
-				else if ( Value.StartsWith( TEXT(" ") ) )
+				else if ( Value.StartsWith( TEXT(" "), ESearchCase::CaseSensitive) )
 				{
 					// Strip leading spaces.
 					do
 					{
-						Value = Value.Right( Value.Len()-1 );
-					} while ( Value.StartsWith( TEXT(" ") ) );
+						Value.RightInline( Value.Len()-1, false );
+					} while ( Value.StartsWith( TEXT(" "), ESearchCase::CaseSensitive) );
 				}
 				else
 				{
diff --git a/Engine/Source/Editor/PropertyEditor/Private/PropertyNode.cpp b/Engine/Source/Editor/PropertyEditor/Private/PropertyNode.cpp
--- a/Engine/Source/Editor/PropertyEditor/Private/PropertyNode.cpp
+++ b/Engine/Source/Editor/PropertyEditor/Private/PropertyNode.cpp
@@ -1829,11 +1829,11 @@ FText FPropertyNode::GetResetToDefaultLabel()
 
 		if ( DefaultValue.Len() > MaxValueLen )
 		{
-			DefaultValue = DefaultValue.Left( MaxValueLen );
+			DefaultValue.LeftInline( MaxValueLen, false );
 			DefaultValue += TEXT( "..." );
 		}
 
-		return FText::Format(NSLOCTEXT("FPropertyNode", "ResetToDefaultLabelFmt", "{0}: {1}"), OutLabel, FText::FromString(DefaultValue));
+		return FText::Format(NSLOCTEXT("FPropertyNode", "ResetToDefaultLabelFmt", "{0}: {1}"), OutLabel, FText::FromString(MoveTemp(DefaultValue)));
 	}
 
 	return OutLabel;
diff --git a/Engine/Source/Editor/PropertyEditor/Private/UserInterface/PropertyEditor/SPropertyMenuActorPicker.cpp b/Engine/Source/Editor/PropertyEditor/Private/UserInterface/PropertyEditor/SPropertyMenuActorPicker.cpp
--- a/Engine/Source/Editor/PropertyEditor/Private/UserInterface/PropertyEditor/SPropertyMenuActorPicker.cpp
+++ b/Engine/Source/Editor/PropertyEditor/Private/UserInterface/PropertyEditor/SPropertyMenuActorPicker.cpp
@@ -160,10 +160,10 @@ bool SPropertyMenuActorPicker::CanPaste()
 
 	FString Class;
 	FString PossibleObjectPath = ClipboardText;
-	if( ClipboardText.Split( TEXT("'"), &Class, &PossibleObjectPath ) )
+	if( ClipboardText.Split( TEXT("'"), &Class, &PossibleObjectPath, ESearchCase::CaseSensitive) )
 	{
 		// Remove the last item
-		PossibleObjectPath = PossibleObjectPath.LeftChop( 1 );
+		PossibleObjectPath.LeftChopInline( 1, false );
 	}
 
 	bool bCanPaste = false;
diff --git a/Engine/Source/Editor/PropertyEditor/Private/UserInterface/PropertyEditor/SPropertyMenuAssetPicker.cpp b/Engine/Source/Editor/PropertyEditor/Private/UserInterface/PropertyEditor/SPropertyMenuAssetPicker.cpp
--- a/Engine/Source/Editor/PropertyEditor/Private/UserInterface/PropertyEditor/SPropertyMenuAssetPicker.cpp
+++ b/Engine/Source/Editor/PropertyEditor/Private/UserInterface/PropertyEditor/SPropertyMenuAssetPicker.cpp
@@ -239,10 +239,10 @@ bool SPropertyMenuAssetPicker::CanPaste()
 
 	FString Class;
 	FString PossibleObjectPath = ClipboardText;
-	if( ClipboardText.Split( TEXT("'"), &Class, &PossibleObjectPath ) )
+	if( ClipboardText.Split( TEXT("'"), &Class, &PossibleObjectPath, ESearchCase::CaseSensitive) )
 	{
 		// Remove the last item
-		PossibleObjectPath = PossibleObjectPath.LeftChop( 1 );
+		PossibleObjectPath.LeftChopInline( 1, false );
 	}
 
 	bool bCanPaste = false;
diff --git a/Engine/Source/Editor/Sequencer/Private/SSequencerLabelBrowser.cpp b/Engine/Source/Editor/Sequencer/Private/SSequencerLabelBrowser.cpp
--- a/Engine/Source/Editor/Sequencer/Private/SSequencerLabelBrowser.cpp
+++ b/Engine/Source/Editor/Sequencer/Private/SSequencerLabelBrowser.cpp
@@ -51,12 +51,12 @@ void SSequencerLabelBrowser::SetSelectedLabel(const FString& Label)
 	Label.ParseIntoArray(SplitLabelStrings, TEXT(" "));
 
 	TArray<FString> SelectedLabels;
-	for (FString LabelString : SplitLabelStrings)
+	for (FString& LabelString : SplitLabelStrings)
 	{
 		if (LabelString.StartsWith(TEXT("label:")))
 		{
-			 LabelString = LabelString.RightChop( 6 );
-			 SelectedLabels.Add(LabelString);
+			 LabelString.RightChopInline( 6, false );
+			 SelectedLabels.Add(MoveTemp(LabelString));
 		}
 	}
 
diff --git a/Engine/Source/Editor/SourceControlWindows/Private/SSourceControlHistory.cpp b/Engine/Source/Editor/SourceControlWindows/Private/SSourceControlHistory.cpp
--- a/Engine/Source/Editor/SourceControlWindows/Private/SSourceControlHistory.cpp
+++ b/Engine/Source/Editor/SourceControlWindows/Private/SSourceControlHistory.cpp
@@ -598,7 +598,7 @@ public:
 			int32 NewLinePos;
 			if (SingleLineDescription.FindChar(TCHAR('\n'), NewLinePos))
 			{
-				SingleLineDescription = SingleLineDescription.Left(NewLinePos);
+				SingleLineDescription.LeftInline(NewLinePos, false);
 			}
 
 			// Trim any trailing new-line characters from the description for the tooltip
diff --git a/Engine/Source/Editor/StatsViewer/Private/StatsEntries/LightingBuildInfo.cpp b/Engine/Source/Editor/StatsViewer/Private/StatsEntries/LightingBuildInfo.cpp
--- a/Engine/Source/Editor/StatsViewer/Private/StatsEntries/LightingBuildInfo.cpp
+++ b/Engine/Source/Editor/StatsViewer/Private/StatsEntries/LightingBuildInfo.cpp
@@ -42,10 +42,10 @@ void ULightingBuildInfo::UpdateNames()
 			LevelName = Object->GetOutermost()->GetName();
 		}
 
-		const int32 NameIndex = LevelName.Find( TEXT("/") );
+		const int32 NameIndex = LevelName.Find( TEXT("/"), ESearchCase::CaseSensitive);
 		if ( NameIndex != INDEX_NONE )
 		{
-			LevelName = LevelName.RightChop( NameIndex + 1 );
+			LevelName.RightChopInline( NameIndex + 1, false );
 		}
 	}
 }
diff --git a/Engine/Source/Editor/StatsViewer/Private/StatsEntries/StaticMeshLightingInfo.cpp b/Engine/Source/Editor/StatsViewer/Private/StatsEntries/StaticMeshLightingInfo.cpp
--- a/Engine/Source/Editor/StatsViewer/Private/StatsEntries/StaticMeshLightingInfo.cpp
+++ b/Engine/Source/Editor/StatsViewer/Private/StatsEntries/StaticMeshLightingInfo.cpp
@@ -37,10 +37,10 @@ void UStaticMeshLightingInfo::UpdateNames()
 			LevelName = TEXT("<None>");
 		}
 
-		const int32 NameIndex = LevelName.Find( TEXT("/") );
+		const int32 NameIndex = LevelName.Find( TEXT("/"), ESearchCase::CaseSensitive);
 		if ( NameIndex != INDEX_NONE )
 		{
-			LevelName = LevelName.RightChop( NameIndex + 1 );
+			LevelName.RightChopInline( NameIndex + 1, false );
 		}
 	}
 
diff --git a/Engine/Source/Editor/StatsViewer/Private/StatsPages/ShaderCookerStatsPage.cpp b/Engine/Source/Editor/StatsViewer/Private/StatsPages/ShaderCookerStatsPage.cpp
--- a/Engine/Source/Editor/StatsViewer/Private/StatsPages/ShaderCookerStatsPage.cpp
+++ b/Engine/Source/Editor/StatsViewer/Private/StatsPages/ShaderCookerStatsPage.cpp
@@ -86,9 +86,9 @@ void FShaderCookerStats::Initialize(uint32 Index)
 		FString FormatName = ShaderPlatformToShaderFormatName((EShaderPlatform)Platform).ToString();
 		if (FormatName.StartsWith(TEXT("SF_")))
 		{
-			FormatName = FormatName.Mid(3);
+			FormatName.MidInline(3, MAX_int32, false);
 		}
-		PlatformNames.Add(FormatName);
+		PlatformNames.Add(MoveTemp(FormatName));
 	}
 	FShaderCookerStatsSet& Set = StatSets[Index];
 	FString CSVData;
diff --git a/Engine/Source/Editor/UnrealEd/Private/Commandlets/AssetRegUtilCommandlet.cpp b/Engine/Source/Editor/UnrealEd/Private/Commandlets/AssetRegUtilCommandlet.cpp
--- a/Engine/Source/Editor/UnrealEd/Private/Commandlets/AssetRegUtilCommandlet.cpp
+++ b/Engine/Source/Editor/UnrealEd/Private/Commandlets/AssetRegUtilCommandlet.cpp
@@ -399,14 +399,14 @@ bool UAssetRegUtilCommandlet::LoadOrderFiles(const FString& OrderFilePath, TSet<
 	OrderFiles.Reserve(Lines.Num());
 	for (int32 EntryIndex = 0; EntryIndex < Lines.Num(); ++EntryIndex)
 	{
-		Lines[EntryIndex].ReplaceInline(TEXT("\r"), TEXT(""));
-		Lines[EntryIndex].ReplaceInline(TEXT("\n"), TEXT(""));
+		Lines[EntryIndex].ReplaceInline(TEXT("\r"), TEXT(""), ESearchCase::CaseSensitive);
+		Lines[EntryIndex].ReplaceInline(TEXT("\n"), TEXT(""), ESearchCase::CaseSensitive);
 		//discard the order number, assuming the list is in-order and has no special bits in use at this stage.
 		int32 QuoteIndex;
 		if (Lines[EntryIndex].FindLastChar('"', QuoteIndex))
 		{
 			FString ReadNum = Lines[EntryIndex].RightChop(QuoteIndex + 1);
-			Lines[EntryIndex] = Lines[EntryIndex].Left(QuoteIndex + 1);
+			Lines[EntryIndex].LeftInline(QuoteIndex + 1, false);
 			//verify our expectations about the order, just in case something changes on the OpenOrder generation side
 			ReadNum.TrimStartInline();
 			if (ReadNum.IsNumeric())
diff --git a/Engine/Source/Editor/UnrealEd/Private/Commandlets/ContentCommandlets.cpp b/Engine/Source/Editor/UnrealEd/Private/Commandlets/ContentCommandlets.cpp
--- a/Engine/Source/Editor/UnrealEd/Private/Commandlets/ContentCommandlets.cpp
+++ b/Engine/Source/Editor/UnrealEd/Private/Commandlets/ContentCommandlets.cpp
@@ -134,7 +134,7 @@ int32 UResavePackagesCommandlet::InitializeResaveParameters( const TArray<FStrin
 		else if (FParse::Value(*CurrentSwitch, TEXT("MAP="), Maps))
 		{
 			// Allow support for -MAP=Value1+Value2+Value3
-			for (int32 PlusIdx = Maps.Find(TEXT("+")); PlusIdx != INDEX_NONE; PlusIdx = Maps.Find(TEXT("+")))
+			for (int32 PlusIdx = Maps.Find(TEXT("+"), ESearchCase::CaseSensitive); PlusIdx != INDEX_NONE; PlusIdx = Maps.Find(TEXT("+"), ESearchCase::CaseSensitive))
 			{
 				const FString NextMap = Maps.Left(PlusIdx);
 				if (NextMap.Len() > 0)
@@ -145,7 +145,7 @@ int32 UResavePackagesCommandlet::InitializeResaveParameters( const TArray<FStrin
 					bExplicitPackages = true;
 				}
 
-				Maps = Maps.Right(Maps.Len() - (PlusIdx + 1));
+				Maps.RightInline(Maps.Len() - (PlusIdx + 1), false);
 			}
 			FString MapFile;
 			FPackageName::SearchForPackageOnDisk(Maps, NULL, &MapFile);
@@ -160,7 +160,7 @@ int32 UResavePackagesCommandlet::InitializeResaveParameters( const TArray<FStrin
 				TArray<FString> Lines;
 
 				// Remove all carriage return characters.
-				Text.ReplaceInline(TEXT("\r"), TEXT(""));
+				Text.ReplaceInline(TEXT("\r"), TEXT(""), ESearchCase::CaseSensitive);
 				// Read all lines
 				Text.ParseIntoArray(Lines, TEXT("\n"), true);
 
@@ -2503,11 +2503,11 @@ int32 UWrangleContentCommandlet::Main( const FString& Params )
 				FString ObjectPathName = It2.Key();
 
 				// skip over the class portion (the It2.Value() has the class pointer already)
-				int32 Space = ObjectPathName.Find(TEXT(" "));
+				int32 Space = ObjectPathName.Find(TEXT(" "), ESearchCase::CaseSensitive);
 				check(Space);
 
 				// get everything after the space
-				ObjectPathName = ObjectPathName.Right(ObjectPathName.Len() - (Space + 1));
+				ObjectPathName.RightInline(ObjectPathName.Len() - (Space + 1), false);
 
 				// load the referenced object
 
@@ -2627,11 +2627,11 @@ int32 UWrangleContentCommandlet::Main( const FString& Params )
 					FString ObjectPathName = ObjectIt.Key();
 
 					// skip over the class portion (the It2.Value() has the class pointer already)
-					int32 Space = ObjectPathName.Find(TEXT(" "));
+					int32 Space = ObjectPathName.Find(TEXT(" "), ESearchCase::CaseSensitive);
 					check(Space > 0);
 
 					// get everything after the space
-					ObjectPathName = ObjectPathName.Right(ObjectPathName.Len() - (Space + 1));
+					ObjectPathName.RightInline(ObjectPathName.Len() - (Space + 1), false);
 
 					// load the unnecessary object
 					UObject* Object = StaticLoadObject(ObjectIt.Value(), NULL, *ObjectPathName, NULL, LOAD_NoWarn, NULL);
diff --git a/Engine/Source/Editor/UnrealEd/Private/Commandlets/CookCommandlet.cpp b/Engine/Source/Editor/UnrealEd/Private/Commandlets/CookCommandlet.cpp
--- a/Engine/Source/Editor/UnrealEd/Private/Commandlets/CookCommandlet.cpp
+++ b/Engine/Source/Editor/UnrealEd/Private/Commandlets/CookCommandlet.cpp
@@ -687,7 +687,7 @@ bool UCookCommandlet::CookByTheBook( const TArray<ITargetPlatform*>& Platforms,
 					const FString ValueElement = ValuesList.Left(PlusIdx);
 					ValueElements.Add(ValueElement);
 
-					ValuesList = ValuesList.Right(ValuesList.Len() - (PlusIdx + 1));
+					ValuesList.RightInline(ValuesList.Len() - (PlusIdx + 1), false);
 				}
 				ValueElements.Add(ValuesList);
 			}
diff --git a/Engine/Source/Editor/UnrealEd/Private/Commandlets/DiffAssetsCommandlet.cpp b/Engine/Source/Editor/UnrealEd/Private/Commandlets/DiffAssetsCommandlet.cpp
--- a/Engine/Source/Editor/UnrealEd/Private/Commandlets/DiffAssetsCommandlet.cpp
+++ b/Engine/Source/Editor/UnrealEd/Private/Commandlets/DiffAssetsCommandlet.cpp
@@ -180,14 +180,14 @@ bool UDiffAssetsCommandlet::ExportFilesToTextAndDiff(const FString& InFilename1,
 		}
 	}
 
-	FString ReplacedDiffCmd = DiffCommand.Replace(TEXT("{1}"), *TextFilename1).Replace(TEXT("{2}"), *TextFilename2);
+	FString ReplacedDiffCmd = DiffCommand.Replace(TEXT("{1}"), *TextFilename1, ESearchCase::CaseSensitive).Replace(TEXT("{2}"), *TextFilename2, ESearchCase::CaseSensitive);
 
-	int32 ArgsAt = DiffCommand.Find(TEXT("{1}")) - 1;
+	int32 ArgsAt = DiffCommand.Find(TEXT("{1}"), ESearchCase::CaseSensitive) - 1;
 	FString Args;
 	if (ArgsAt > 0)
 	{
 		Args = *ReplacedDiffCmd + ArgsAt + 1;
-		ReplacedDiffCmd = ReplacedDiffCmd.Left(ArgsAt);
+		ReplacedDiffCmd.LeftInline(ArgsAt, false);
 	}
 
 	if (!FPlatformProcess::CreateProc(*ReplacedDiffCmd, *Args, true, false, false, NULL, 0, NULL, NULL).IsValid())
diff --git a/Engine/Source/Editor/UnrealEd/Private/Commandlets/GatherTextFromSourceCommandlet.cpp b/Engine/Source/Editor/UnrealEd/Private/Commandlets/GatherTextFromSourceCommandlet.cpp
--- a/Engine/Source/Editor/UnrealEd/Private/Commandlets/GatherTextFromSourceCommandlet.cpp
+++ b/Engine/Source/Editor/UnrealEd/Private/Commandlets/GatherTextFromSourceCommandlet.cpp
@@ -1243,7 +1243,8 @@ void UGatherTextFromSourceCommandlet::FIfDescriptor::TryParse(const FString& Tex
 		// Handle "#if defined <defname>"
 		if (RemainingText.StartsWith(DefinedString, ESearchCase::CaseSensitive))
 		{
-			RemainingText = RemainingText.RightChop(DefinedString.Len()).TrimStart();
+			RemainingText.RightChopInline(DefinedString.Len(), false);
+			RemainingText.TrimStartInline();
 		}
 
 		Context.PushMacroBlock(RemainingText);
@@ -1279,7 +1280,8 @@ void UGatherTextFromSourceCommandlet::FElIfDescriptor::TryParse(const FString& T
 		// Handle "#elif defined <defname>"
 		if (RemainingText.StartsWith(DefinedString, ESearchCase::CaseSensitive))
 		{
-			RemainingText = RemainingText.RightChop(DefinedString.Len()).TrimStart();
+			RemainingText.RightChopInline(DefinedString.Len(), false);
+			RemainingText.TrimStartInline();
 		}
 
 		Context.PopMacroBlock(); // Pop the current #if or #ifdef state
@@ -1324,7 +1326,7 @@ bool UGatherTextFromSourceCommandlet::FMacroDescriptor::ParseArgsFromMacro(const
 	FString RemainingText = Text.RightChop(GetToken().Len());
 	RemainingText.TrimStartInline();
 
-	const int32 OpenBracketIdx = RemainingText.Find(TEXT("("));
+	const int32 OpenBracketIdx = RemainingText.Find(TEXT("("), ESearchCase::CaseSensitive);
 	if (OpenBracketIdx == INDEX_NONE)
 	{
 		// No opening bracket; warn about this, but don't consider it an error as we're likely parsing something we shouldn't be
diff --git a/Engine/Source/Editor/UnrealEd/Private/Commandlets/ShaderPipelineCacheToolsCommandlet.cpp b/Engine/Source/Editor/UnrealEd/Private/Commandlets/ShaderPipelineCacheToolsCommandlet.cpp
--- a/Engine/Source/Editor/UnrealEd/Private/Commandlets/ShaderPipelineCacheToolsCommandlet.cpp
+++ b/Engine/Source/Editor/UnrealEd/Private/Commandlets/ShaderPipelineCacheToolsCommandlet.cpp
@@ -917,13 +917,13 @@ void ParseQuoteComma(const FString& InLine, TArray<FString>& OutParts)
 		{
 			break;
 		}
-		Line = Line.RightChop(QuoteLoc + 1);
+		Line.RightChopInline(QuoteLoc + 1, false);
 		if (!Line.FindChar(TCHAR('\"'), QuoteLoc))
 		{
 			break;
 		}
 		OutParts.Add(Line.Left(QuoteLoc));
-		Line = Line.RightChop(QuoteLoc + 1);
+		Line.RightChopInline(QuoteLoc + 1, false);
 	}
 }
 
diff --git a/Engine/Source/Editor/UnrealEd/Private/CookOnTheFlyServer.cpp b/Engine/Source/Editor/UnrealEd/Private/CookOnTheFlyServer.cpp
--- a/Engine/Source/Editor/UnrealEd/Private/CookOnTheFlyServer.cpp
+++ b/Engine/Source/Editor/UnrealEd/Private/CookOnTheFlyServer.cpp
@@ -3705,7 +3705,7 @@ private:
 		int32 ParamIndex = InOutParams.Find(InParamToRemove);
 		if (ParamIndex >= 0)
 		{
-			int32 NextParamIndex = InOutParams.Find(TEXT(" -"), ESearchCase::IgnoreCase, ESearchDir::FromStart, ParamIndex + 1);
+			int32 NextParamIndex = InOutParams.Find(TEXT(" -"), ESearchCase::CaseSensitive, ESearchDir::FromStart, ParamIndex + 1);
 			if (NextParamIndex < ParamIndex)
 			{
 				NextParamIndex = InOutParams.Len();
diff --git a/Engine/Source/Editor/UnrealEd/Private/EditorEngine.cpp b/Engine/Source/Editor/UnrealEd/Private/EditorEngine.cpp
--- a/Engine/Source/Editor/UnrealEd/Private/EditorEngine.cpp
+++ b/Engine/Source/Editor/UnrealEd/Private/EditorEngine.cpp
@@ -7157,7 +7157,7 @@ bool FActorLabelUtilities::SplitActorLabel(FString& InOutLabel, int32& OutIdx)
 			FString Idx = InOutLabel.RightChop(CharIdx);
 			if (Idx.Len() > 0)
 			{
-				InOutLabel = InOutLabel.Left(CharIdx);
+				InOutLabel.LeftInline(CharIdx);
 				OutIdx = FCString::Atoi(*Idx);
 				return true;
 			}
diff --git a/Engine/Source/Editor/UnrealEd/Private/EditorObject.cpp b/Engine/Source/Editor/UnrealEd/Private/EditorObject.cpp
--- a/Engine/Source/Editor/UnrealEd/Private/EditorObject.cpp
+++ b/Engine/Source/Editor/UnrealEd/Private/EditorObject.cpp
@@ -267,7 +267,7 @@ static const TCHAR* ImportProperties(
 			}
 			if (Length != StrLine.Len())
 			{
-				StrLine = StrLine.Left(Length);
+				StrLine.LeftInline(Length, false);
 			}
 		}
 
diff --git a/Engine/Source/Editor/UnrealEd/Private/EditorServer.cpp b/Engine/Source/Editor/UnrealEd/Private/EditorServer.cpp
--- a/Engine/Source/Editor/UnrealEd/Private/EditorServer.cpp
+++ b/Engine/Source/Editor/UnrealEd/Private/EditorServer.cpp
@@ -408,11 +408,11 @@ bool UEditorEngine::SafeExec( UWorld* InWorld, const TCHAR* InStr, FOutputDevice
 				{
 					break;
 				}
-				ObjectName = ObjectName.Mid( i+1 );
+				ObjectName.MidInline( i+1, MAX_int32, false );
 			}
 			if( ObjectName.Find(TEXT("."), ESearchCase::CaseSensitive)>=0 )
 			{
-				ObjectName = ObjectName.Left( ObjectName.Find(TEXT(".")) );
+				ObjectName.LeftInline( ObjectName.Find(TEXT("."), ESearchCase::CaseSensitive), false );
 			}
 		}
 
@@ -5440,10 +5440,10 @@ void ListMapPackageDependencies(const TCHAR* InStr)
 				{
 					// get package name of the import
 					FString ImportPackage = FPackageName::FilenameToLongPackageName(Linker->GetImportPathName(ImportIdx));
-					int32 PeriodIdx = ImportPackage.Find(TEXT("."));
+					int32 PeriodIdx = ImportPackage.Find(TEXT("."), ESearchCase::CaseSensitive);
 					if (PeriodIdx != INDEX_NONE)
 					{
-						ImportPackage = ImportPackage.Left(PeriodIdx);
+						ImportPackage.LeftInline(PeriodIdx, false);
 					}
 					ReferencedPackages.Add(ImportPackage, true);
 				}
diff --git a/Engine/Source/Editor/UnrealEd/Private/Factories/PackFactory.cpp b/Engine/Source/Editor/UnrealEd/Private/Factories/PackFactory.cpp
--- a/Engine/Source/Editor/UnrealEd/Private/Factories/PackFactory.cpp
+++ b/Engine/Source/Editor/UnrealEd/Private/Factories/PackFactory.cpp
@@ -483,14 +483,14 @@ UObject* UPackFactory::FactoryCreateBinary
 					FString DestFilename = It.Filename();
 					if (DestFilename.StartsWith(TEXT("Source/")))
 					{
-						DestFilename = DestFilename.RightChop(7);
+						DestFilename.RightChopInline(7, false);
 					}
 					else 
 					{
 						const int32 SourceIndex = DestFilename.Find(TEXT("/Source/"));
 						if (SourceIndex != INDEX_NONE)
 						{
-							DestFilename = DestFilename.RightChop(SourceIndex + 8);
+							DestFilename.RightChopInline(SourceIndex + 8, false);
 						}
 					}
 
@@ -525,14 +525,14 @@ UObject* UPackFactory::FactoryCreateBinary
 					FString DestFilename = It.Filename();
 					if (DestFilename.StartsWith(TEXT("Content/")))
 					{
-						DestFilename = DestFilename.RightChop(8);
+						DestFilename.RightChopInline(8, false);
 					}
 					else
 					{
 						const int32 ContentIndex = DestFilename.Find(ContentFolder);
 						if (ContentIndex != INDEX_NONE)
 						{
-							DestFilename = DestFilename.RightChop(ContentIndex + 9);
+							DestFilename.RightChopInline(ContentIndex + 9, false);
 						}
 					}
 					DestFilename = ContentDestinationRoot / DestFilename;
@@ -572,14 +572,14 @@ UObject* UPackFactory::FactoryCreateBinary
 					FString DestFilename = FileToCopy;
 					if (DestFilename.StartsWith(TEXT("Source/")))
 					{
-						DestFilename = DestFilename.RightChop(7);
+						DestFilename.RightChopInline(7, false);
 					}
 					else 
 					{
 						const int32 SourceIndex = DestFilename.Find(TEXT("/Source/"));
 						if (SourceIndex != INDEX_NONE)
 						{
-							DestFilename = DestFilename.RightChop(SourceIndex + 8);
+							DestFilename.RightChopInline(SourceIndex + 8, false);
 						}
 					}
 					DestFilename = SourceModuleInfo.ModuleSourcePath / DestFilename;
@@ -624,14 +624,14 @@ UObject* UPackFactory::FactoryCreateBinary
 					FString DestFilename = FileToCopy;
 					if (DestFilename.StartsWith(TEXT("Content/")))
 					{
-						DestFilename = DestFilename.RightChop(8);
+						DestFilename.RightChopInline(8, false);
 					}
 					else
 					{
 						const int32 ContentIndex = DestFilename.Find(ContentFolder);
 						if (ContentIndex != INDEX_NONE)
 						{
-							DestFilename = DestFilename.RightChop(ContentIndex + 9);
+							DestFilename.RightChopInline(ContentIndex + 9, false);
 						}
 					}
 					DestFilename = ContentDestinationRoot / DestFilename;
diff --git a/Engine/Source/Editor/UnrealEd/Private/Fbx/FbxMainImport.cpp b/Engine/Source/Editor/UnrealEd/Private/Fbx/FbxMainImport.cpp
--- a/Engine/Source/Editor/UnrealEd/Private/Fbx/FbxMainImport.cpp
+++ b/Engine/Source/Editor/UnrealEd/Private/Fbx/FbxMainImport.cpp
@@ -2154,7 +2154,7 @@ void FFbxImporter::ConvertLodPrefixToLodGroup()
 
 		//Get a valid name for the LODGroup actor
 		FString FbxNodeName = UTF8_TO_TCHAR(FirstNode->GetName());
-		FbxNodeName = FbxNodeName.RightChop(5);
+		FbxNodeName.RightChopInline(5, false);
 		FbxNodeName += TEXT("_LodGroup");
 		//Create a LodGroup and child all fbx node to the Group
 		FbxNode* ActorNode = FbxNode::Create(Scene, TCHAR_TO_UTF8(*FbxNodeName));
@@ -2169,7 +2169,7 @@ void FFbxImporter::ConvertLodPrefixToLodGroup()
 				if (bCanReduce)
 				{
 					FString FbxGeneratedNodeName = UTF8_TO_TCHAR(FirstNode->GetName());
-					FbxGeneratedNodeName = FbxGeneratedNodeName.RightChop(5);
+					FbxGeneratedNodeName.RightChopInline(5, false);
 					FbxGeneratedNodeName += TEXT(GeneratedLODNameSuffix) + FString::FromInt(CurrentLodIndex);
 					//Generated LOD add dummy FbxNode to tell the import to add such a LOD
 					FbxNode* DummyGeneratedLODActorNode = FbxNode::Create(Scene, TCHAR_TO_UTF8(*FbxGeneratedNodeName));
diff --git a/Engine/Source/Editor/UnrealEd/Private/Fbx/FbxMaterialImport.cpp b/Engine/Source/Editor/UnrealEd/Private/Fbx/FbxMaterialImport.cpp
--- a/Engine/Source/Editor/UnrealEd/Private/Fbx/FbxMaterialImport.cpp
+++ b/Engine/Source/Editor/UnrealEd/Private/Fbx/FbxMaterialImport.cpp
@@ -433,7 +433,7 @@ FString UnFbx::FFbxImporter::GetMaterialFullName(FbxSurfaceMaterial& FbxMaterial
 			if (SkinXXNumber.IsNumeric())
 			{
 				// remove the '_skinXX' suffix from the material name					
-				MaterialFullName = MaterialFullName.LeftChop(MaterialFullName.Len() - Offset);
+				MaterialFullName.LeftChopInline(MaterialFullName.Len() - Offset, false);
 			}
 		}
 	}
diff --git a/Engine/Source/Editor/UnrealEd/Private/Fbx/FbxSkeletalMeshExport.cpp b/Engine/Source/Editor/UnrealEd/Private/Fbx/FbxSkeletalMeshExport.cpp
--- a/Engine/Source/Editor/UnrealEd/Private/Fbx/FbxSkeletalMeshExport.cpp
+++ b/Engine/Source/Editor/UnrealEd/Private/Fbx/FbxSkeletalMeshExport.cpp
@@ -608,7 +608,7 @@ void ExportObjectMetadataToBones(const UObject* ObjectToExport, const TArray<Fbx
 					NodeName = TagAsString.Left(CharPos);
 
 					// The remaining part is the actual metadata tag
-					TagAsString = TagAsString.RightChop(CharPos + 1); // exclude the period
+					TagAsString.RightChopInline(CharPos + 1, false); // exclude the period
 				}
 
 				// Try to attach the metadata to its associated node by name
diff --git a/Engine/Source/Editor/UnrealEd/Private/Fbx/FbxSkeletalMeshImport.cpp b/Engine/Source/Editor/UnrealEd/Private/Fbx/FbxSkeletalMeshImport.cpp
--- a/Engine/Source/Editor/UnrealEd/Private/Fbx/FbxSkeletalMeshImport.cpp
+++ b/Engine/Source/Editor/UnrealEd/Private/Fbx/FbxSkeletalMeshImport.cpp
@@ -4252,7 +4252,7 @@ void UnFbx::FFbxImporter::ImportMorphTargetsInternal( TArray<FbxNode*>& SkelMesh
 						// Maya adds the name of the blendshape and an underscore to the front of the channel name, so remove it
 						if(ChannelName.StartsWith(BlendShapeName))
 						{
-							ChannelName = ChannelName.Right(ChannelName.Len() - (BlendShapeName.Len()+1));
+							ChannelName.RightInline(ChannelName.Len() - (BlendShapeName.Len()+1), false);
 						}
 
 						for(int32 ShapeIndex = 0; ShapeIndex<CurrentChannelShapeCount; ++ShapeIndex)
diff --git a/Engine/Source/Editor/UnrealEd/Private/Fbx/FbxStaticMeshImport.cpp b/Engine/Source/Editor/UnrealEd/Private/Fbx/FbxStaticMeshImport.cpp
--- a/Engine/Source/Editor/UnrealEd/Private/Fbx/FbxStaticMeshImport.cpp
+++ b/Engine/Source/Editor/UnrealEd/Private/Fbx/FbxStaticMeshImport.cpp
@@ -2165,7 +2165,7 @@ static void FindMeshSockets( FbxNode* StartNode, TArray<FbxSocketNode>& OutFbxSo
 			if( SocketName.StartsWith( SocketPrefix ) )
 			{
 				// Remove the prefix from the name
-				SocketName = SocketName.RightChop( SocketPrefix.Len() );
+				SocketName.RightChopInline( SocketPrefix.Len(), false );
 
 				FbxSocketNode NewNode;
 				NewNode.Node = StartNode;
diff --git a/Engine/Source/Editor/UnrealEd/Private/Fbx/ReimportFbxSceneFactory.cpp b/Engine/Source/Editor/UnrealEd/Private/Fbx/ReimportFbxSceneFactory.cpp
--- a/Engine/Source/Editor/UnrealEd/Private/Fbx/ReimportFbxSceneFactory.cpp
+++ b/Engine/Source/Editor/UnrealEd/Private/Fbx/ReimportFbxSceneFactory.cpp
@@ -787,7 +787,7 @@ void UReimportFbxSceneFactory::RecursivelySetComponentProperties(USCS_Node* Curr
 	FString ReduceNodeName = NodeName;
 	if (NameContainTemplateSuffixe)
 	{
-		ReduceNodeName = ReduceNodeName.Left(IndexTemplateSuffixe);
+		ReduceNodeName.LeftInline(IndexTemplateSuffixe, false);
 	}
 
 	USceneComponent *CurrentNodeSceneComponent = Cast<USceneComponent>(CurrentNodeActorComponent);
diff --git a/Engine/Source/Editor/UnrealEd/Private/FbxAnimUtils.cpp b/Engine/Source/Editor/UnrealEd/Private/FbxAnimUtils.cpp
--- a/Engine/Source/Editor/UnrealEd/Private/FbxAnimUtils.cpp
+++ b/Engine/Source/Editor/UnrealEd/Private/FbxAnimUtils.cpp
@@ -152,7 +152,7 @@ namespace FbxAnimUtils
 										// Maya adds the name of the blendshape and an underscore to the front of the channel name, so remove it
 										if (ChannelName.StartsWith(BlendShapeName))
 										{
-											ChannelName = ChannelName.Right(ChannelName.Len() - (BlendShapeName.Len() + 1));
+											ChannelName.RightInline(ChannelName.Len() - (BlendShapeName.Len() + 1), false);
 										}
 
 										FbxAnimCurve* Curve = Geometry->GetShapeChannel(BlendShapeIndex, ChannelIndex, (FbxAnimLayer*)AnimStack->GetMember(0));
diff --git a/Engine/Source/Editor/UnrealEd/Private/FileHelpers.cpp b/Engine/Source/Editor/UnrealEd/Private/FileHelpers.cpp
--- a/Engine/Source/Editor/UnrealEd/Private/FileHelpers.cpp
+++ b/Engine/Source/Editor/UnrealEd/Private/FileHelpers.cpp
@@ -2075,7 +2075,7 @@ void FEditorFileUtils::OpenLevelPickingDialog(const FOnLevelsChosen& OnLevelsCho
 				// Remove the slash if needed
 				if ( FilesystemPath.EndsWith(TEXT("/"), ESearchCase::CaseSensitive) )
 				{
-					FilesystemPath = FilesystemPath.LeftChop(1);
+					FilesystemPath.LeftChopInline(1, false);
 				}
 
 				FEditorDirectories::Get().SetLastDirectory(ELastDirectory::LEVEL, FilesystemPath);
diff --git a/Engine/Source/Editor/UnrealEd/Private/Kismet2/BlueprintEditorUtils.cpp b/Engine/Source/Editor/UnrealEd/Private/Kismet2/BlueprintEditorUtils.cpp
--- a/Engine/Source/Editor/UnrealEd/Private/Kismet2/BlueprintEditorUtils.cpp
+++ b/Engine/Source/Editor/UnrealEd/Private/Kismet2/BlueprintEditorUtils.cpp
@@ -7281,7 +7281,7 @@ FName FBlueprintEditorUtils::FindUniqueKismetName(const UBlueprint* InBlueprint,
 		// If the length of the final string will be too long, cut off the end so we can fit the number
 		if(CountLength + BaseName.Len() > NameValidator->GetMaximumNameLength())
 		{
-			BaseName = BaseName.Left(NameValidator->GetMaximumNameLength() - CountLength);
+			BaseName.LeftInline(NameValidator->GetMaximumNameLength() - CountLength, false);
 		}
 		KismetName = FString::Printf(TEXT("%s_%d"), *BaseName, Count);
 		Count++;
diff --git a/Engine/Source/Editor/UnrealEd/Private/Kismet2/CompilerResultsLog.cpp b/Engine/Source/Editor/UnrealEd/Private/Kismet2/CompilerResultsLog.cpp
--- a/Engine/Source/Editor/UnrealEd/Private/Kismet2/CompilerResultsLog.cpp
+++ b/Engine/Source/Editor/UnrealEd/Private/Kismet2/CompilerResultsLog.cpp
@@ -529,11 +529,12 @@ TArray< TSharedRef<FTokenizedMessage> > FCompilerResultsLog::ParseCompilerLogDum
 	for (int32 i = 0; i < MessageLines.Num(); ++i)
 	{
 		FString Line = MessageLines[i];
-		if (Line.EndsWith(TEXT("\r")))
+		if (Line.EndsWith(TEXT("\r"), ESearchCase::CaseSensitive))
 		{
-			Line = Line.LeftChop(1);
+			Line.LeftChopInline(1, false);
 		}
-		Line = Line.ConvertTabsToSpaces(4).TrimEnd();
+		Line.ConvertTabsToSpacesInline(4);
+		Line.TrimEndInline();
 
 		// handle output line error message if applicable
 		// @todo Handle case where there are parenthesis in path names
@@ -594,7 +595,7 @@ void FCompilerResultsLog::OnGotoError(const TSharedRef<IMessageToken>& Token)
 	FString FullPath, LineNumberString;
 	if (Token->ToText().ToString().Split(TEXT("("), &FullPath, &LineNumberString, ESearchCase::CaseSensitive))
 	{
-		LineNumberString = LineNumberString.LeftChop(1); // remove right parenthesis
+		LineNumberString.LeftChopInline(1, false); // remove right parenthesis
 		int32 LineNumber = FCString::Strtoi(*LineNumberString, NULL, 10);
 
 		FSourceCodeNavigation::OpenSourceFile( FullPath, LineNumber );
diff --git a/Engine/Source/Editor/UnrealEd/Private/Kismet2/ComponentEditorUtils.cpp b/Engine/Source/Editor/UnrealEd/Private/Kismet2/ComponentEditorUtils.cpp
--- a/Engine/Source/Editor/UnrealEd/Private/Kismet2/ComponentEditorUtils.cpp
+++ b/Engine/Source/Editor/UnrealEd/Private/Kismet2/ComponentEditorUtils.cpp
@@ -238,14 +238,14 @@ FString FComponentEditorUtils::GenerateValidVariableName(TSubclassOf<UActorCompo
 	FString SuffixToStrip( TEXT( "Component" ) );
 	if( ComponentTypeName.EndsWith( SuffixToStrip ) )
 	{
-		ComponentTypeName = ComponentTypeName.Left( ComponentTypeName.Len() - SuffixToStrip.Len() );
+		ComponentTypeName.LeftInline( ComponentTypeName.Len() - SuffixToStrip.Len(), false );
 	}
 
 	// Strip off 'Actor' if the class ends with that so as not to confuse actors with components
 	SuffixToStrip = TEXT( "Actor" );
 	if( ComponentTypeName.EndsWith( SuffixToStrip ) )
 	{
-		ComponentTypeName = ComponentTypeName.Left( ComponentTypeName.Len() - SuffixToStrip.Len() );
+		ComponentTypeName.LeftInline( ComponentTypeName.Len() - SuffixToStrip.Len(), false );
 	}
 
 	// Try to create a name without any numerical suffix first
diff --git a/Engine/Source/Editor/UnrealEd/Private/PackageBackup.cpp b/Engine/Source/Editor/UnrealEd/Private/PackageBackup.cpp
--- a/Engine/Source/Editor/UnrealEd/Private/PackageBackup.cpp
+++ b/Engine/Source/Editor/UnrealEd/Private/PackageBackup.cpp
@@ -182,7 +182,7 @@ bool FAutoPackageBackup::ShouldBackupPackage( const UPackage& InPackage, FString
 			FString CurBackupFileNameExtension = FPaths::GetExtension(CurBackupFileName);
 
 			// The base file name of the backup file is going to include an underscore followed by a timestamp, so they must be removed for comparison's sake
-			CurBackupBaseFileName = CurBackupBaseFileName.Left( CurBackupBaseFileName.Find( TEXT("_"), ESearchCase::CaseSensitive, ESearchDir::FromEnd ) );
+			CurBackupBaseFileName.LeftInline( CurBackupBaseFileName.Find( TEXT("_"), ESearchCase::CaseSensitive, ESearchDir::FromEnd ), false );
 					
 			// If the base file names and extensions match, we've found a backup
 			if ( CurBackupBaseFileName == ExistingBaseFileName &&  CurBackupFileNameExtension == ExistingFileNameExtension )
