Matinee: Allow null pointer if no open tab valid

diff --git a/Engine/Source/Editor/Matinee/Private/Matinee.cpp b/Engine/Source/Editor/Matinee/Private/Matinee.cpp
--- a/Engine/Source/Editor/Matinee/Private/Matinee.cpp
+++ b/Engine/Source/Editor/Matinee/Private/Matinee.cpp
@@ -967,7 +967,7 @@ void FMatinee::InitMatinee(const EToolkitMode::Type Mode, const TSharedPtr< clas
 				)	
 		);
 
-	FGlobalTabmanager::Get()->RestoreFrom( Layout, TSharedPtr<SWindow>() );
+	FGlobalTabmanager::Get()->RestoreFrom( Layout, TSharedPtr<SWindow>(), false, EOutputCanBeNullptr::IfNoOpenTabValid );
 
 	// OK, we're now initialized!
 	bIsInitialized = true;
