Matinee: Remove UE4 -> UE

diff --git a/Engine/Source/Editor/Matinee/Private/MatineeMenus.cpp b/Engine/Source/Editor/Matinee/Private/MatineeMenus.cpp
--- a/Engine/Source/Editor/Matinee/Private/MatineeMenus.cpp
+++ b/Engine/Source/Editor/Matinee/Private/MatineeMenus.cpp
@@ -526,7 +526,7 @@ void FMatinee::NewGroupPopupTextCommitted(
 
 		// Folders don't need a group instance
 		UInterpGroupInst* NewGroupInst = NULL;
-		//@todo UE4 Matinee: No Kismet
+		//@todo UE Matinee: No Kismet
 		//USeqVar_Character* AIVarObj = NULL;
 		if( !bIsNewFolder )
 		{
diff --git a/Engine/Source/Editor/Matinee/Private/MatineeTools.cpp b/Engine/Source/Editor/Matinee/Private/MatineeTools.cpp
--- a/Engine/Source/Editor/Matinee/Private/MatineeTools.cpp
+++ b/Engine/Source/Editor/Matinee/Private/MatineeTools.cpp
@@ -3537,7 +3537,7 @@ void FMatinee::SaveRecordingParentOffsets(void)
 			}
 
 			//get the parent actor
-			/** @todo UE4 no longer using general base attachment system
+			/** @todo UE no longer using general base attachment system
 			AActor* ParentActor = Actor->Base;
 			if (ParentActor)
 			{
@@ -3571,7 +3571,7 @@ void FMatinee::ApplyRecordingParentOffsets(void)
 				return;
 			}
 
-			/** @todo UE4 no longer using general base attachment system
+			/** @todo UE no longer using general base attachment system
 			//get the parent actor
 			AActor* ParentActor = Actor->Base;
 			if (ParentActor)
@@ -3582,7 +3582,7 @@ void FMatinee::ApplyRecordingParentOffsets(void)
 			*/
 		}
 
-		/** @todo UE4 no longer using general base attachment system
+		/** @todo UE no longer using general base attachment system
 		//now apply parent offsets to list
 		for (int32 i = 0; i < RecordingActorsWithParents.Num(); ++i)
 		{
diff --git a/Engine/Source/Editor/Matinee/Private/MatineeTrackView.cpp b/Engine/Source/Editor/Matinee/Private/MatineeTrackView.cpp
--- a/Engine/Source/Editor/Matinee/Private/MatineeTrackView.cpp
+++ b/Engine/Source/Editor/Matinee/Private/MatineeTrackView.cpp
@@ -60,7 +60,7 @@ public:
 				.Padding(4, 0, 0, 0)
 				[
 					SNew(SRichTextBlock)
-					.Text(NSLOCTEXT("Matinee", "MatineeLastVersionSupported", "As of 4.23, Matinee is no longer supported by UE4 and will be removed from the engine in the near future. Once removed, you will <NormalText.Important>no longer be able to run a Matinee or open Matinee Editor</>.\nPlease use the <a id=\"udn\" href=\"/Engine/Sequencer/HowTo/MatineeConversionTool\" style=\"DarkHyperlink\">Matinee to Sequencer Conversion Tool</> to convert any files to Sequencer as soon as possible."))
+					.Text(NSLOCTEXT("Matinee", "MatineeLastVersionSupported", "As of 4.23, Matinee is no longer supported by UE and will be removed from the engine in the near future. Once removed, you will <NormalText.Important>no longer be able to run a Matinee or open Matinee Editor</>.\nPlease use the <a id=\"udn\" href=\"/Engine/Sequencer/HowTo/MatineeConversionTool\" style=\"DarkHyperlink\">Matinee to Sequencer Conversion Tool</> to convert any files to Sequencer as soon as possible."))
 					.AutoWrapText(true)
 					.DecoratorStyleSet(&FEditorStyle::Get())
 					+ SRichTextBlock::HyperlinkDecorator(TEXT("udn"), FSlateHyperlinkRun::FOnClick::CreateStatic(&OnDocLinkClicked))
