Fixed issue where Matinee Editor was asserting during window closure. This is fallout from CL 16391884, where we defined a new OnClose() event that had a name collision in Matinee's class and thus OnClose() was getting called twice.

This is a minimal workaround, simply resolving the function name conflict without affecting existing order of operations, preferred since Matinee is slated for removal.

#rnx

diff --git a/Engine/Source/Editor/Matinee/Private/Matinee.cpp b/Engine/Source/Editor/Matinee/Private/Matinee.cpp
--- a/Engine/Source/Editor/Matinee/Private/Matinee.cpp
+++ b/Engine/Source/Editor/Matinee/Private/Matinee.cpp
@@ -990,7 +990,7 @@ void FMatinee::InitMatinee(const EToolkitMode::Type Mode, const TSharedPtr< clas
 
 FMatinee::~FMatinee()
 {
-	OnClose();
+	OnClosed();
 
 	RestoreLevelViewports();
 
@@ -1030,7 +1030,7 @@ void FMatinee::AddReferencedObjects( FReferenceCollector& Collector )
 		Collector.AddReferencedObject( Pair.Value.TrackHelper );
 	}
 
-	// Check for non-NULL, as these references will be cleared in OnClose.
+	// Check for non-NULL, as these references will be cleared in OnClosed.
 	if ( TrackWindow.IsValid() && TrackWindow->InterpEdVC.IsValid() )
 	{
 		TrackWindow->InterpEdVC->AddReferencedObjects( Collector );
@@ -1991,8 +1991,7 @@ int32 FMatinee::GetNumCameraActors(void) const
 	return CameraCount;
 }
 
-
-void FMatinee::OnClose()
+void FMatinee::OnClosed()
 {
 	// Safely stop recording if it is in progress
 	if (IsRecordingInterpValues())
diff --git a/Engine/Source/Editor/Matinee/Private/Matinee.h b/Engine/Source/Editor/Matinee/Private/Matinee.h
--- a/Engine/Source/Editor/Matinee/Private/Matinee.h
+++ b/Engine/Source/Editor/Matinee/Private/Matinee.h
@@ -69,7 +69,7 @@ public:
 
 	TSharedRef<SDockTab> SpawnTab( const FSpawnTabArgs& TabSpawnArgs, FName TabIdentifier );
 
-	void OnClose();
+	void OnClosed();
 
 	/** Extends the toolbar for Matinee specific options */
 	void ExtendDefaultToolbarMenu();
